Program to test certain aspects of platform abstraction.

Currently deals with FProcHandles only.

diff --git a/Engine/Source/Programs/TestPAL/Private/Main.cpp b/Engine/Source/Programs/TestPAL/Private/Main.cpp
new file mode 100644
--- /dev/null
+++ b/Engine/Source/Programs/TestPAL/Private/Main.cpp
@@ -0,0 +1,113 @@
+// Copyright 1998-2014 Epic Games, Inc. All Rights Reserved.
+
+#include "PrivatePCH.h"
+#include "RequiredProgramMainCPPInclude.h"
+
+DEFINE_LOG_CATEGORY(LogTestPAL);
+
+IMPLEMENT_APPLICATION(TestPAL, "TestPAL");
+
+#define CHILD_INVOCATION_ARG			"-child"
+
+namespace
+{
+	/**
+	 * Checks if this instance is nested (i.e. run under itself)
+	 *
+	 * @param ArgC Number of commandline arguments.
+	 * @param ArgV Commandline arguments
+	 * @return true if nested
+	 */
+	bool IsNestedInstance(int32 ArgC, ANSICHAR* ArgV[])
+	{
+		for (int32 IdxArg = 0; IdxArg < ArgC; ++IdxArg)
+		{
+			if (!FCStringAnsi::Strcmp(ArgV[IdxArg], CHILD_INVOCATION_ARG))
+			{
+				return true;
+			}
+		}
+
+		return false;
+	}
+}
+
+int32 RunAsChild(const TCHAR* CommandLine)
+{
+	FPlatformMisc::SetCrashHandler(NULL);
+	FPlatformMisc::SetGracefulTerminationHandler();
+
+	GEngineLoop.PreInit(CommandLine);
+
+	// set a random delay pretending to do some useful work up to a minute. 
+	double RandomWorkTime = FMath::FRandRange(0.0f, 60.0f);
+
+	UE_LOG(LogTestPAL, Display, TEXT("Running as child (pid %d), will be doing work for %f seconds."), FPlatformProcess::GetCurrentProcessId(), RandomWorkTime);
+
+	double StartTime = FPlatformTime::Seconds();
+
+	// Use all the CPU!
+	for (;;)
+	{
+		double CurrentTime = FPlatformTime::Seconds();
+		if (CurrentTime - StartTime >= RandomWorkTime)
+		{
+			break;
+		}
+	}
+
+	UE_LOG(LogTestPAL, Display, TEXT("Child (pid %d) finished work."), FPlatformProcess::GetCurrentProcessId());
+
+	FEngineLoop::AppPreExit();
+	FEngineLoop::AppExit();
+	return 0;
+}
+
+int32 RunAsParent(const TCHAR* CommandLine)
+{
+	FPlatformMisc::SetCrashHandler(NULL);
+	FPlatformMisc::SetGracefulTerminationHandler();
+
+	GEngineLoop.PreInit(CommandLine);
+	UE_LOG(LogTestPAL, Display, TEXT("Running as parent."));
+
+	// Run slave instance continuously
+	int NumChildrenToSpawn = 255, MaxAtOnce = 5;
+	FParent Parent(NumChildrenToSpawn, MaxAtOnce);
+
+	Parent.Run();
+
+	UE_LOG(LogTestPAL, Display, TEXT("Parent quit."));
+
+	FEngineLoop::AppPreExit();
+	FEngineLoop::AppExit();
+	return 0;
+}
+
+int32 main(int32 ArgC, char* ArgV[])
+{
+	FString CommandLine;
+	for (int32 Option = 1; Option < ArgC; Option++)
+	{
+		CommandLine += TEXT(" ");
+		FString Argument(ANSI_TO_TCHAR(ArgV[Option]));
+		if (Argument.Contains(TEXT(" ")))
+		{
+			if (Argument.Contains(TEXT("=")))
+			{
+				FString ArgName;
+				FString ArgValue;
+				Argument.Split( TEXT("="), &ArgName, &ArgValue );
+				Argument = FString::Printf( TEXT("%s=\"%s\""), *ArgName, *ArgValue );
+			}
+			else
+			{
+				Argument = FString::Printf(TEXT("\"%s\""), *Argument);
+			}
+		}
+		CommandLine += Argument;
+	}
+
+	// fork and run under self
+	return IsNestedInstance(ArgC, ArgV) ? RunAsChild(*CommandLine) : RunAsParent(*CommandLine);
+}
diff --git a/Engine/Source/Programs/TestPAL/Private/Parent.cpp b/Engine/Source/Programs/TestPAL/Private/Parent.cpp
new file mode 100644
--- /dev/null
+++ b/Engine/Source/Programs/TestPAL/Private/Parent.cpp
@@ -0,0 +1,63 @@
+#include "PrivatePCH.h"
+#include "Parent.h"
+
+FParent::FParent(int InNumTotalChildren, int InMaxChildrenAtOnce)
+	:	NumTotalChildren(InNumTotalChildren)
+	,	MaxChildrenAtOnce(InMaxChildrenAtOnce)
+{
+
+}
+
+FProcHandle FParent::Launch()
+{
+	// Launch the worker process
+	int32 PriorityModifier = -1; // below normal
+
+	uint32 WorkerId = 0;
+	FString WorkerName = FPlatformProcess::ExecutableName(false);
+	FProcHandle WorkerHandle = FPlatformProcess::CreateProc(*WorkerName, TEXT("-child"), true, false, false, &WorkerId, PriorityModifier, NULL, NULL);
+	if (!WorkerHandle.IsValid())
+	{
+		// If this doesn't error, the app will hang waiting for jobs that can never be completed
+		UE_LOG(LogTestPAL, Fatal, TEXT("Couldn't launch %s! Make sure the file is in your binaries folder."), *WorkerName);
+	}
+
+	return WorkerHandle;
+}
+
+void FParent::Run()
+{
+	while (NumTotalChildren > 0)
+	{
+		// see if there are any new children to spawn
+		while (Children.Num() < MaxChildrenAtOnce)
+		{
+			UE_LOG(LogTestPAL, Log, TEXT("Launching a child (%d more to go)."), NumTotalChildren - 1);
+			FProcHandle Child = Launch();
+			UE_LOG(LogTestPAL, Log, TEXT("Launch successful"));
+
+			Children.Add(Child);
+			--NumTotalChildren;
+		}
+
+		// sleep for a while
+		FPlatformProcess::Sleep(0.5f);
+
+		// see if any children have finished
+		for (int ChildIdx = 0; ChildIdx < Children.Num();)
+		{
+			int32 ReturnCode = -1;
+			if (FPlatformProcess::GetProcReturnCode(Children[ChildIdx], &ReturnCode))
+			{
+				// print its return code
+				UE_LOG(LogTestPAL, Log, TEXT("Child finished, return code %d"), ReturnCode);
+
+				Children.RemoveAt(ChildIdx);
+			}
+			else
+			{
+				++ChildIdx;
+			}
+		}
+	}
+}
diff --git a/Engine/Source/Programs/TestPAL/Private/Parent.h b/Engine/Source/Programs/TestPAL/Private/Parent.h
new file mode 100644
--- /dev/null
+++ b/Engine/Source/Programs/TestPAL/Private/Parent.h
@@ -0,0 +1,27 @@
+// Copyright 1998-2014 Epic Games, Inc. All Rights Reserved.
+
+#pragma once
+
+DECLARE_LOG_CATEGORY_EXTERN(LogTestPAL, Verbose, All);
+
+class FParent
+{
+	/** Children to be controlled */
+	TArray< FProcHandle > Children;
+
+	/** Total children to spawn */
+	int NumTotalChildren;
+
+	/** Max children at once */
+	int MaxChildrenAtOnce;
+
+	/** Launches the child, returns the launched process handle. */
+	FProcHandle Launch();
+
+public:
+
+	/** Spawns the children one by one */
+	void Run();
+
+	FParent(int InNumTotalChildren, int InMaxChildrenAtOnce);
+};
diff --git a/Engine/Source/Programs/TestPAL/Private/PrivatePCH.h b/Engine/Source/Programs/TestPAL/Private/PrivatePCH.h
new file mode 100644
--- /dev/null
+++ b/Engine/Source/Programs/TestPAL/Private/PrivatePCH.h
@@ -0,0 +1,9 @@
+// Copyright 1998-2014 Epic Games, Inc. All Rights Reserved.
+
+#pragma once
+
+// other modules' headers
+#include "Core.h"
+
+// this module headers
+#include "Parent.h"
diff --git a/Engine/Source/Programs/TestPAL/TestPAL.Build.cs b/Engine/Source/Programs/TestPAL/TestPAL.Build.cs
new file mode 100644
--- /dev/null
+++ b/Engine/Source/Programs/TestPAL/TestPAL.Build.cs
@@ -0,0 +1,20 @@
+﻿// Copyright 1998-2014 Epic Games, Inc. All Rights Reserved.
+
+using UnrealBuildTool;
+
+public class TestPAL : ModuleRules
+{
+	public TestPAL(TargetInfo Target)
+	{
+		PublicIncludePaths.Add("Runtime/Launch/Public");
+
+		PrivateIncludePaths.Add("Runtime/Launch/Private");		// For LaunchEngineLoop.cpp include
+
+		PrivateDependencyModuleNames.AddRange(
+			new string[] {
+				"Core",
+				"Projects"
+			}
+		);
+	}
+}
diff --git a/Engine/Source/Programs/TestPAL/TestPAL.Target.cs b/Engine/Source/Programs/TestPAL/TestPAL.Target.cs
new file mode 100644
--- /dev/null
+++ b/Engine/Source/Programs/TestPAL/TestPAL.Target.cs
@@ -0,0 +1,60 @@
+﻿// Copyright 1998-2014 Epic Games, Inc. All Rights Reserved.
+
+using UnrealBuildTool;
+using System.Collections.Generic;
+
+public class TestPALTarget : TargetRules
+{
+	public TestPALTarget(TargetInfo Target)
+	{
+		Type = TargetType.Program;
+	}
+
+	//
+	// TargetRules interface.
+	//
+
+	public override bool GetSupportedPlatforms(ref List<UnrealTargetPlatform> OutPlatforms)
+	{
+		// It is valid for only server platforms
+		return UnrealBuildTool.UnrealBuildTool.GetAllServerPlatforms(ref OutPlatforms, false);
+	}
+
+	public override void SetupBinaries(
+		TargetInfo Target,
+		ref List<UEBuildBinaryConfiguration> OutBuildBinaryConfigurations,
+		ref List<string> OutExtraModuleNames
+		)
+	{
+		OutBuildBinaryConfigurations.Add(
+			new UEBuildBinaryConfiguration(	InType: UEBuildBinaryType.Executable,
+											InModuleNames: new List<string>() { "TestPAL" } )
+			);
+	}
+
+	public override bool ShouldCompileMonolithic(UnrealTargetPlatform InPlatform, UnrealTargetConfiguration InConfiguration)
+	{
+		return true;
+	}
+
+	public override void SetupGlobalEnvironment(
+		TargetInfo Target,
+		ref LinkEnvironmentConfiguration OutLinkEnvironmentConfiguration,
+		ref CPPEnvironmentConfiguration OutCPPEnvironmentConfiguration
+		)
+	{
+		// Lean and mean
+		UEBuildConfiguration.bCompileLeanAndMeanUE = true;
+
+		// No editor or editor-only data is needed
+		UEBuildConfiguration.bBuildEditor = false;
+		//UEBuildConfiguration.bBuildWithEditorOnlyData = false;
+
+		// Compile out references from Core to the rest of the engine
+		UEBuildConfiguration.bCompileAgainstEngine = false;	// compiling without engine is broken (overriden functions do not override base class)
+		UEBuildConfiguration.bCompileAgainstCoreUObject = false;
+
+		// Make a console application under Windows, so entry point is main() everywhere
+		OutLinkEnvironmentConfiguration.bIsBuildingConsoleApplication = true;
+	}
+}
