Fix MessageBox not displaying long strings (UE-9176).

- Also added a test: ./TestPAL messagebox

diff --git a/Engine/Source/Programs/TestPAL/Private/Main.cpp b/Engine/Source/Programs/TestPAL/Private/Main.cpp
--- a/Engine/Source/Programs/TestPAL/Private/Main.cpp
+++ b/Engine/Source/Programs/TestPAL/Private/Main.cpp
@@ -116,6 +116,28 @@ int32 CaseTest(const TCHAR* CommandLine)
 	return 0;
 }
 
+/**
+ * Message box test/
+ */
+int32 MessageBoxTest(const TCHAR* CommandLine)
+{
+	FPlatformMisc::SetCrashHandler(NULL);
+	FPlatformMisc::SetGracefulTerminationHandler();
+
+	GEngineLoop.PreInit(CommandLine);
+	UE_LOG(LogTestPAL, Display, TEXT("Running message box test."));
+	
+	FString Display(TEXT("I am a big big string in a big big game, it's not a big big thing if you print me. But I do do feel that I do do will be displayed wrong, displayed wrong...  or not."));
+	FString Caption(TEXT("I am a big big caption in a big big game, it's not a big big thing if you print me. But I do do feel that I do do will be displayed wrong, displayed wrong... or not."));
+	EAppReturnType::Type Result = FPlatformMisc::MessageBoxExt(EAppMsgType::YesNo, *Display, *Caption);
+
+	UE_LOG(LogTestPAL, Display, TEXT("MessageBoxExt result: %d."), static_cast<int32>(Result));
+	
+	FEngineLoop::AppPreExit();
+	FEngineLoop::AppExit();
+	return 0;
+}
+
 /**
  * Selects and runs one of test cases.
  *
@@ -139,6 +161,10 @@ int32 MultiplexedMain(int32 ArgC, char* ArgV[])
 		{
 			return CaseTest(*TestPAL::CommandLine);
 		}
+		else if (!FCStringAnsi::Strcmp(ArgV[IdxArg], ARG_MESSAGEBOX_TEST))
+		{
+			return MessageBoxTest(*TestPAL::CommandLine);
+		}
 	}
 
 	FPlatformMisc::SetCrashHandler(NULL);
diff --git a/Engine/Source/Programs/TestPAL/Private/PrivatePCH.h b/Engine/Source/Programs/TestPAL/Private/PrivatePCH.h
--- a/Engine/Source/Programs/TestPAL/Private/PrivatePCH.h
+++ b/Engine/Source/Programs/TestPAL/Private/PrivatePCH.h
@@ -8,6 +8,7 @@
 #define ARG_PROC_TEST						"proc"
 #define ARG_PROC_TEST_CHILD					"proc-child"
 #define ARG_CASE_SENSITIVITY_TEST			"case"
+#define ARG_MESSAGEBOX_TEST					"messagebox"
 
 
 // this module headers
diff --git a/Engine/Source/Runtime/Core/Private/Linux/LinuxPlatformMisc.cpp b/Engine/Source/Runtime/Core/Private/Linux/LinuxPlatformMisc.cpp
--- a/Engine/Source/Runtime/Core/Private/Linux/LinuxPlatformMisc.cpp
+++ b/Engine/Source/Runtime/Core/Private/Linux/LinuxPlatformMisc.cpp
@@ -442,12 +442,14 @@ EAppReturnType::Type FLinuxPlatformMisc::MessageBoxExt(EAppMsgType::Type MsgType
 			break;
 	}
 
+	FTCHARToUTF8 CaptionUTF8(Caption);
+	FTCHARToUTF8 TextUTF8(Text);
 	SDL_MessageBoxData MessageBoxData = 
 	{
 		SDL_MESSAGEBOX_INFORMATION,
 		NULL, // No parent window
-		TCHAR_TO_UTF8(Caption),
-		TCHAR_TO_UTF8(Text),
+		CaptionUTF8.Get(),
+		TextUTF8.Get(),
 		NumberOfButtons,
 		Buttons,
 		NULL // Default color scheme
