TestPAL: add a test suite for DirectoryWatcher.

- Uncovers a problem with DWLinux - files added/removed in a subdirectory get reported with an incorrect path.

diff --git a/Engine/Source/Programs/TestPAL/Private/Main.cpp b/Engine/Source/Programs/TestPAL/Private/Main.cpp
--- a/Engine/Source/Programs/TestPAL/Private/Main.cpp
+++ b/Engine/Source/Programs/TestPAL/Private/Main.cpp
@@ -2,6 +2,7 @@
 
 #include "PrivatePCH.h"
 #include "RequiredProgramMainCPPInclude.h"
+#include "TestDirectoryWatcher.h"
 
 DEFINE_LOG_CATEGORY(LogTestPAL);
 
@@ -166,6 +167,10 @@ int32 MultiplexedMain(int32 ArgC, char* ArgV[])
 		{
 			return MessageBoxTest(*TestPAL::CommandLine);
 		}
+		else if (!FCStringAnsi::Strcmp(ArgV[IdxArg], ARG_DIRECTORY_WATCHER_TEST))
+		{
+			return DirectoryWatcherTest(*TestPAL::CommandLine);
+		}
 	}
 
 	FPlatformMisc::SetCrashHandler(NULL);
@@ -173,6 +178,16 @@ int32 MultiplexedMain(int32 ArgC, char* ArgV[])
 	
 	GEngineLoop.PreInit(*TestPAL::CommandLine);
 	UE_LOG(LogTestPAL, Warning, TEXT("Unable to find any known test name, no test started."));
+
+	UE_LOG(LogTestPAL, Warning, TEXT(""));
+	UE_LOG(LogTestPAL, Warning, TEXT("Available test cases:"));
+	UE_LOG(LogTestPAL, Warning, TEXT("  %s: test process handling API"), ARG_PROC_TEST);
+	UE_LOG(LogTestPAL, Warning, TEXT("  %s: test case-insensitive file operations"), ARG_CASE_SENSITIVITY_TEST);
+	UE_LOG(LogTestPAL, Warning, TEXT("  %s: test message box bug (too long strings)"), ARG_MESSAGEBOX_TEST);
+	UE_LOG(LogTestPAL, Warning, TEXT("  %s: test directory watcher"), ARG_DIRECTORY_WATCHER_TEST);
+	UE_LOG(LogTestPAL, Warning, TEXT(""));
+	UE_LOG(LogTestPAL, Warning, TEXT("Pass one of those to run an appropriate test."));
+
 	FEngineLoop::AppPreExit();
 	FEngineLoop::AppExit();
 	return 1;
diff --git a/Engine/Source/Programs/TestPAL/Private/PrivatePCH.h b/Engine/Source/Programs/TestPAL/Private/PrivatePCH.h
--- a/Engine/Source/Programs/TestPAL/Private/PrivatePCH.h
+++ b/Engine/Source/Programs/TestPAL/Private/PrivatePCH.h
@@ -9,7 +9,7 @@
 #define ARG_PROC_TEST_CHILD					"proc-child"
 #define ARG_CASE_SENSITIVITY_TEST			"case"
 #define ARG_MESSAGEBOX_TEST					"messagebox"
-
+#define ARG_DIRECTORY_WATCHER_TEST			"dirwatcher"
 
 // this module headers
 #include "Parent.h"
diff --git a/Engine/Source/Programs/TestPAL/Private/TestDirectoryWatcher.cpp b/Engine/Source/Programs/TestPAL/Private/TestDirectoryWatcher.cpp
new file mode 100644
--- /dev/null
+++ b/Engine/Source/Programs/TestPAL/Private/TestDirectoryWatcher.cpp
@@ -0,0 +1,99 @@
+// Copyright 1998-2015 Epic Games, Inc. All Rights Reserved.
+
+#include "PrivatePCH.h"
+#include "IDirectoryWatcher.h"
+#include "DirectoryWatcherModule.h"
+#include "TestDirectoryWatcher.h"
+
+struct FChangeDetector
+{
+	void OnDirectoryChanged(const TArray<FFileChangeData>& FileChanges)
+	{
+		UE_LOG(LogTestPAL, Display, TEXT("%d change(s) detected"), static_cast<int32>(FileChanges.Num()));
+
+		int ChangeIdx = 0;
+		for (const auto& ThisEntry : FileChanges)
+		{
+			UE_LOG(LogTestPAL, Display, TEXT("  Change %d: %s was %s"),
+				++ChangeIdx,
+				*ThisEntry.Filename,
+				ThisEntry.Action == FFileChangeData::FCA_Added ? TEXT("added") :
+					(ThisEntry.Action == FFileChangeData::FCA_Removed ? TEXT("removed") :
+						(ThisEntry.Action == FFileChangeData::FCA_Modified ? TEXT("modified") : TEXT("??? (unknown)")
+						)
+					)
+			);
+		}
+	}
+};
+
+/**
+ * Kicks directory watcher test/
+ */
+int32 DirectoryWatcherTest(const TCHAR* CommandLine)
+{
+	FPlatformMisc::SetCrashHandler(NULL);
+	FPlatformMisc::SetGracefulTerminationHandler();
+
+	GEngineLoop.PreInit(CommandLine);
+	UE_LOG(LogTestPAL, Display, TEXT("Running directory watcher test."));
+
+	IPlatformFile& PlatformFile = FPlatformFileManager::Get().GetPlatformFile();
+	FString TestDir = FString::Printf(TEXT("/tmp/DirectoryWatcherTest%d"), FPlatformProcess::GetCurrentProcessId());
+
+	if (PlatformFile.CreateDirectory(*TestDir) && PlatformFile.CreateDirectory(*(TestDir + TEXT("/subtest"))))
+	{
+		FChangeDetector Detector;
+		FDelegateHandle DirectoryChangedHandle;
+
+		IDirectoryWatcher* DirectoryWatcher = FModuleManager::Get().LoadModuleChecked<FDirectoryWatcherModule>(TEXT("DirectoryWatcher")).Get();
+		if (DirectoryWatcher)
+		{
+			auto Callback = IDirectoryWatcher::FDirectoryChanged::CreateRaw(&Detector, &FChangeDetector::OnDirectoryChanged);
+			DirectoryWatcher->RegisterDirectoryChangedCallback_Handle(TestDir, Callback, DirectoryChangedHandle);
+		}
+		else
+		{
+			UE_LOG(LogTestPAL, Fatal, TEXT("Could not get DirectoryWatcher module"));
+		}
+
+		FPlatformProcess::Sleep(1.0f);
+		DirectoryWatcher->Tick(1.0f);
+
+		// create and remove directory
+		UE_LOG(LogTestPAL, Log, TEXT("Creating '%s'"), *(TestDir+TEXT("/test")));
+		PlatformFile.CreateDirectory(*(TestDir+TEXT("/test")));
+		FPlatformProcess::Sleep(1.0f);
+		DirectoryWatcher->Tick(1.0f);
+
+		UE_LOG(LogTestPAL, Log, TEXT("Deleting '%s'"), *(TestDir+TEXT("/test")));
+		PlatformFile.DeleteDirectory(*(TestDir+TEXT("/test")));
+		FPlatformProcess::Sleep(1.0f);
+		DirectoryWatcher->Tick(1.0f);
+
+		// create and remove in a sub directory
+		UE_LOG(LogTestPAL, Log, TEXT("Creating '%s'"), *(TestDir+TEXT("/subtest/blah")));
+		PlatformFile.CreateDirectory(*(TestDir+TEXT("/subtest/blah")));
+		FPlatformProcess::Sleep(1.0f);
+		DirectoryWatcher->Tick(1.0f);
+
+		UE_LOG(LogTestPAL, Log, TEXT("Deleting '%s'"), *(TestDir+TEXT("/subtest/blah")));
+		PlatformFile.DeleteDirectory(*(TestDir+TEXT("/subtest/blah")));
+		FPlatformProcess::Sleep(1.0f);
+		DirectoryWatcher->Tick(1.0f);
+
+		// clean up
+		DirectoryWatcher->UnregisterDirectoryChangedCallback_Handle(TestDir, DirectoryChangedHandle);
+		// remove dirs as well
+		PlatformFile.DeleteDirectory(*(TestDir + TEXT("/subtest")));
+		PlatformFile.DeleteDirectory(*TestDir);
+	}
+	else
+	{
+		UE_LOG(LogTestPAL, Fatal, TEXT("Could not create test directory %s."), *TestDir);
+	}
+
+	FEngineLoop::AppPreExit();
+	FEngineLoop::AppExit();
+	return 0;
+}
diff --git a/Engine/Source/Programs/TestPAL/Private/TestDirectoryWatcher.h b/Engine/Source/Programs/TestPAL/Private/TestDirectoryWatcher.h
new file mode 100644
--- /dev/null
+++ b/Engine/Source/Programs/TestPAL/Private/TestDirectoryWatcher.h
@@ -0,0 +1,5 @@
+// Copyright 1998-2015 Epic Games, Inc. All Rights Reserved.
+
+#pragma once
+
+int32 DirectoryWatcherTest(const TCHAR* CommandLine);
\ No newline at end of file
diff --git a/Engine/Source/Programs/TestPAL/TestPAL.Build.cs b/Engine/Source/Programs/TestPAL/TestPAL.Build.cs
--- a/Engine/Source/Programs/TestPAL/TestPAL.Build.cs
+++ b/Engine/Source/Programs/TestPAL/TestPAL.Build.cs
@@ -13,7 +13,8 @@ public class TestPAL : ModuleRules
 		PrivateDependencyModuleNames.AddRange(
 			new string[] {
 				"Core",
-				"Projects"
+				"Projects",
+				"DirectoryWatcher"
 			}
 		);
 	}
