Thread singleton test for JarekS' change in CL 2360391.

diff --git a/Engine/Source/Programs/TestPAL/Private/Main.cpp b/Engine/Source/Programs/TestPAL/Private/Main.cpp
--- a/Engine/Source/Programs/TestPAL/Private/Main.cpp
+++ b/Engine/Source/Programs/TestPAL/Private/Main.cpp
@@ -140,6 +140,95 @@ int32 MessageBoxTest(const TCHAR* CommandLine)
 	return 0;
 }
 
+/**
+ * ************  Thread singleton test *****************
+ */
+
+/**
+ * Per-thread singleton
+ */
+struct FPerThreadTestSingleton : public TThreadSingleton<FPerThreadTestSingleton>
+{
+	FPerThreadTestSingleton()
+	{
+		UE_LOG(LogTestPAL, Log, TEXT("FPerThreadTestSingleton (this=%p) created for thread %d"),
+			this,
+			FPlatformTLS::GetCurrentThreadId());
+	}
+
+	void DoSomething()
+	{
+		UE_LOG(LogTestPAL, Log, TEXT("Thread %d is about to quit"), FPlatformTLS::GetCurrentThreadId());
+	}
+
+	virtual ~FPerThreadTestSingleton()
+	{
+		UE_LOG(LogTestPAL, Log, TEXT("FPerThreadTestSingleton (%p) destroyed for thread %d"),
+			this,
+			FPlatformTLS::GetCurrentThreadId());
+	}
+};
+
+//DECLARE_THREAD_SINGLETON( FPerThreadTestSingleton );
+
+/**
+ * Thread runnable
+ */
+struct FSingletonTestingThread : public FRunnable
+{
+	virtual uint32 Run()
+	{
+		FPerThreadTestSingleton& Dummy = FPerThreadTestSingleton::Get();
+
+		FPlatformProcess::Sleep(3.0f);
+
+		Dummy.DoSomething();
+		return 0;
+	}
+};
+
+/**
+ * Thread singleton test
+ */
+int32 ThreadSingletonTest(const TCHAR* CommandLine)
+{
+	FPlatformMisc::SetCrashHandler(NULL);
+	FPlatformMisc::SetGracefulTerminationHandler();
+
+	GEngineLoop.PreInit(CommandLine);
+	UE_LOG(LogTestPAL, Display, TEXT("Running thread singleton test."));
+
+	const int kNumTestThreads = 10;
+
+	FSingletonTestingThread * RunnableArray[kNumTestThreads] = { nullptr };
+	FRunnableThread * ThreadArray[kNumTestThreads] = { nullptr };
+
+	// start all threads
+	for (int Idx = 0; Idx < kNumTestThreads; ++Idx)
+	{
+		RunnableArray[Idx] = new FSingletonTestingThread();
+		ThreadArray[Idx] = FRunnableThread::Create(RunnableArray[Idx],
+			*FString::Printf(TEXT("TestThread%d"), Idx));
+	}
+
+	GLog->FlushThreadedLogs();
+	GLog->Flush();
+
+	// join all threads
+	for (int Idx = 0; Idx < kNumTestThreads; ++Idx)
+	{
+		ThreadArray[Idx]->WaitForCompletion();
+		delete ThreadArray[Idx];
+		ThreadArray[Idx] = nullptr;
+		delete RunnableArray[Idx];
+		RunnableArray[Idx] = nullptr;
+	}
+
+	FEngineLoop::AppPreExit();
+	FEngineLoop::AppExit();
+	return 0;
+}
+
 /**
  * Selects and runs one of test cases.
  *
@@ -171,6 +260,10 @@ int32 MultiplexedMain(int32 ArgC, char* ArgV[])
 		{
 			return DirectoryWatcherTest(*TestPAL::CommandLine);
 		}
+		else if (!FCStringAnsi::Strcmp(ArgV[IdxArg], ARG_THREAD_SINGLETON_TEST))
+		{
+			return ThreadSingletonTest(*TestPAL::CommandLine);
+		}
 	}
 
 	FPlatformMisc::SetCrashHandler(NULL);
@@ -185,6 +278,7 @@ int32 MultiplexedMain(int32 ArgC, char* ArgV[])
 	UE_LOG(LogTestPAL, Warning, TEXT("  %s: test case-insensitive file operations"), ARG_CASE_SENSITIVITY_TEST);
 	UE_LOG(LogTestPAL, Warning, TEXT("  %s: test message box bug (too long strings)"), ARG_MESSAGEBOX_TEST);
 	UE_LOG(LogTestPAL, Warning, TEXT("  %s: test directory watcher"), ARG_DIRECTORY_WATCHER_TEST);
+	UE_LOG(LogTestPAL, Warning, TEXT("  %s: test per-thread singletons"), ARG_THREAD_SINGLETON_TEST);
 	UE_LOG(LogTestPAL, Warning, TEXT(""));
 	UE_LOG(LogTestPAL, Warning, TEXT("Pass one of those to run an appropriate test."));
 
diff --git a/Engine/Source/Programs/TestPAL/Private/PrivatePCH.h b/Engine/Source/Programs/TestPAL/Private/PrivatePCH.h
--- a/Engine/Source/Programs/TestPAL/Private/PrivatePCH.h
+++ b/Engine/Source/Programs/TestPAL/Private/PrivatePCH.h
@@ -10,6 +10,7 @@
 #define ARG_CASE_SENSITIVITY_TEST			"case"
 #define ARG_MESSAGEBOX_TEST					"messagebox"
 #define ARG_DIRECTORY_WATCHER_TEST			"dirwatcher"
+#define ARG_THREAD_SINGLETON_TEST			"threadsingleton"
 
 // this module headers
 #include "Parent.h"
