PalTest - Misc fix for print

diff --git a/Engine/Source/Programs/TestPAL/Private/Main.cpp b/Engine/Source/Programs/TestPAL/Private/Main.cpp
--- a/Engine/Source/Programs/TestPAL/Private/Main.cpp
+++ b/Engine/Source/Programs/TestPAL/Private/Main.cpp
@@ -274,11 +274,11 @@ int32 MultiplexedMain(int32 ArgC, char* ArgV[])
 
 	UE_LOG(LogTestPAL, Warning, TEXT(""));
 	UE_LOG(LogTestPAL, Warning, TEXT("Available test cases:"));
-	UE_LOG(LogTestPAL, Warning, TEXT("  %s: test process handling API"), ARG_PROC_TEST);
-	UE_LOG(LogTestPAL, Warning, TEXT("  %s: test case-insensitive file operations"), ARG_CASE_SENSITIVITY_TEST);
-	UE_LOG(LogTestPAL, Warning, TEXT("  %s: test message box bug (too long strings)"), ARG_MESSAGEBOX_TEST);
-	UE_LOG(LogTestPAL, Warning, TEXT("  %s: test directory watcher"), ARG_DIRECTORY_WATCHER_TEST);
-	UE_LOG(LogTestPAL, Warning, TEXT("  %s: test per-thread singletons"), ARG_THREAD_SINGLETON_TEST);
+	UE_LOG(LogTestPAL, Warning, TEXT("  %s: test process handling API"), ANSI_TO_TCHAR( ARG_PROC_TEST ));
+	UE_LOG(LogTestPAL, Warning, TEXT("  %s: test case-insensitive file operations"), ANSI_TO_TCHAR( ARG_CASE_SENSITIVITY_TEST ));
+	UE_LOG(LogTestPAL, Warning, TEXT("  %s: test message box bug (too long strings)"), ANSI_TO_TCHAR( ARG_MESSAGEBOX_TEST ));
+	UE_LOG(LogTestPAL, Warning, TEXT("  %s: test directory watcher"), ANSI_TO_TCHAR( ARG_DIRECTORY_WATCHER_TEST ));
+	UE_LOG(LogTestPAL, Warning, TEXT("  %s: test per-thread singletons"), ANSI_TO_TCHAR( ARG_THREAD_SINGLETON_TEST ));
 	UE_LOG(LogTestPAL, Warning, TEXT(""));
 	UE_LOG(LogTestPAL, Warning, TEXT("Pass one of those to run an appropriate test."));
 
