TestPAL: more extensive test for DirectoryWatcher.

diff --git a/Engine/Source/Programs/TestPAL/Private/TestDirectoryWatcher.cpp b/Engine/Source/Programs/TestPAL/Private/TestDirectoryWatcher.cpp
--- a/Engine/Source/Programs/TestPAL/Private/TestDirectoryWatcher.cpp
+++ b/Engine/Source/Programs/TestPAL/Private/TestDirectoryWatcher.cpp
@@ -9,12 +9,12 @@ struct FChangeDetector
 {
 	void OnDirectoryChanged(const TArray<FFileChangeData>& FileChanges)
 	{
-		UE_LOG(LogTestPAL, Display, TEXT("%d change(s) detected"), static_cast<int32>(FileChanges.Num()));
+		UE_LOG(LogTestPAL, Display, TEXT("  -- %d change(s) detected"), static_cast<int32>(FileChanges.Num()));
 
 		int ChangeIdx = 0;
 		for (const auto& ThisEntry : FileChanges)
 		{
-			UE_LOG(LogTestPAL, Display, TEXT("  Change %d: %s was %s"),
+			UE_LOG(LogTestPAL, Display, TEXT("      Change %d: %s was %s"),
 				++ChangeIdx,
 				*ThisEntry.Filename,
 				ThisEntry.Action == FFileChangeData::FCA_Added ? TEXT("added") :
@@ -39,7 +39,7 @@ int32 DirectoryWatcherTest(const TCHAR* CommandLine)
 	UE_LOG(LogTestPAL, Display, TEXT("Running directory watcher test."));
 
 	IPlatformFile& PlatformFile = FPlatformFileManager::Get().GetPlatformFile();
-	FString TestDir = FString::Printf(TEXT("/tmp/DirectoryWatcherTest%d"), FPlatformProcess::GetCurrentProcessId());
+	FString TestDir = FString::Printf(TEXT("%sDirectoryWatcherTest%d"), FPlatformProcess::UserTempDir(), FPlatformProcess::GetCurrentProcessId());
 
 	if (PlatformFile.CreateDirectory(*TestDir) && PlatformFile.CreateDirectory(*(TestDir + TEXT("/subtest"))))
 	{
@@ -51,6 +51,7 @@ int32 DirectoryWatcherTest(const TCHAR* CommandLine)
 		{
 			auto Callback = IDirectoryWatcher::FDirectoryChanged::CreateRaw(&Detector, &FChangeDetector::OnDirectoryChanged);
 			DirectoryWatcher->RegisterDirectoryChangedCallback_Handle(TestDir, Callback, DirectoryChangedHandle);
+			UE_LOG(LogTestPAL, Display, TEXT("Registered callback for changes in '%s'"), *TestDir);
 		}
 		else
 		{
@@ -61,32 +62,124 @@ int32 DirectoryWatcherTest(const TCHAR* CommandLine)
 		DirectoryWatcher->Tick(1.0f);
 
 		// create and remove directory
-		UE_LOG(LogTestPAL, Log, TEXT("Creating '%s'"), *(TestDir+TEXT("/test")));
-		PlatformFile.CreateDirectory(*(TestDir+TEXT("/test")));
+		UE_LOG(LogTestPAL, Display, TEXT("Creating DIRECTORY '%s'"), *(TestDir + TEXT("/test")));
+		verify(PlatformFile.CreateDirectory(*(TestDir + TEXT("/test"))));
+		DirectoryWatcher->Tick(1.0f);
 		FPlatformProcess::Sleep(1.0f);
 		DirectoryWatcher->Tick(1.0f);
 
-		UE_LOG(LogTestPAL, Log, TEXT("Deleting '%s'"), *(TestDir+TEXT("/test")));
-		PlatformFile.DeleteDirectory(*(TestDir+TEXT("/test")));
+		UE_LOG(LogTestPAL, Display, TEXT("Deleting DIRECTORY '%s'"), *(TestDir + TEXT("/test")));
+		verify(PlatformFile.DeleteDirectory(*(TestDir + TEXT("/test"))));
+		DirectoryWatcher->Tick(1.0f);
 		FPlatformProcess::Sleep(1.0f);
 		DirectoryWatcher->Tick(1.0f);
 
 		// create and remove in a sub directory
-		UE_LOG(LogTestPAL, Log, TEXT("Creating '%s'"), *(TestDir+TEXT("/subtest/blah")));
-		PlatformFile.CreateDirectory(*(TestDir+TEXT("/subtest/blah")));
+		UE_LOG(LogTestPAL, Display, TEXT("Creating DIRECTORY '%s'"), *(TestDir + TEXT("/subtest/blah")));
+		verify(PlatformFile.CreateDirectory(*(TestDir + TEXT("/subtest/blah"))));
+		DirectoryWatcher->Tick(1.0f);
 		FPlatformProcess::Sleep(1.0f);
 		DirectoryWatcher->Tick(1.0f);
 
-		UE_LOG(LogTestPAL, Log, TEXT("Deleting '%s'"), *(TestDir+TEXT("/subtest/blah")));
-		PlatformFile.DeleteDirectory(*(TestDir+TEXT("/subtest/blah")));
+		UE_LOG(LogTestPAL, Display, TEXT("Deleting DIRECTORY '%s'"), *(TestDir + TEXT("/subtest/blah")));
+		verify(PlatformFile.DeleteDirectory(*(TestDir + TEXT("/subtest/blah"))));
+		DirectoryWatcher->Tick(1.0f);
 		FPlatformProcess::Sleep(1.0f);
 		DirectoryWatcher->Tick(1.0f);
 
+		{
+			// create file
+			FString DummyFileName = TestDir + TEXT("/test file.bin");
+			UE_LOG(LogTestPAL, Display, TEXT("Creating FILE '%s'"), *DummyFileName);
+			IFileHandle* DummyFile = PlatformFile.OpenWrite(*DummyFileName);
+			check(DummyFile);
+			DirectoryWatcher->Tick(1.0f);
+			FPlatformProcess::Sleep(1.0f);
+			DirectoryWatcher->Tick(1.0f);
+
+			// modify file
+			UE_LOG(LogTestPAL, Display, TEXT("Modifying FILE '%s'"), *DummyFileName);
+			uint8 Contents = 0;
+			DummyFile->Write(&Contents, sizeof(Contents));
+			DirectoryWatcher->Tick(1.0f);
+			FPlatformProcess::Sleep(1.0f);
+			DirectoryWatcher->Tick(1.0f);
+
+			// close the file
+			UE_LOG(LogTestPAL, Display, TEXT("Closing FILE '%s'"), *DummyFileName);
+			delete DummyFile;
+			DummyFile = nullptr;
+			DirectoryWatcher->Tick(1.0f);
+			FPlatformProcess::Sleep(1.0f);
+			DirectoryWatcher->Tick(1.0f);
+
+			// delete file
+			UE_LOG(LogTestPAL, Display, TEXT("Deleting FILE '%s'"), *DummyFileName);
+			verify(PlatformFile.DeleteFile(*DummyFileName));
+			DirectoryWatcher->Tick(1.0f);
+			FPlatformProcess::Sleep(1.0f);
+			DirectoryWatcher->Tick(1.0f);
+		}
+
+		// now the same in a grandchild directory
+		{
+			FString GrandChildDir = TestDir + TEXT("/subtest/grandchild");
+
+			UE_LOG(LogTestPAL, Display, TEXT("Creating DIRECTORY '%s'"), *GrandChildDir);
+			verify(PlatformFile.CreateDirectory(*GrandChildDir));
+			DirectoryWatcher->Tick(1.0f);
+			FPlatformProcess::Sleep(1.0f);
+			DirectoryWatcher->Tick(1.0f);
+
+			{
+				// create file
+				FString DummyFileName = GrandChildDir + TEXT("/test file.bin");
+				UE_LOG(LogTestPAL, Display, TEXT("Creating FILE '%s'"), *DummyFileName);
+				IFileHandle* DummyFile = PlatformFile.OpenWrite(*DummyFileName);
+				check(DummyFile);
+				DirectoryWatcher->Tick(1.0f);
+				FPlatformProcess::Sleep(1.0f);
+				DirectoryWatcher->Tick(1.0f);
+
+				// modify file
+				UE_LOG(LogTestPAL, Display, TEXT("Modifying FILE '%s'"), *DummyFileName);
+				uint8 Contents = 0;
+				DummyFile->Write(&Contents, sizeof(Contents));
+				DirectoryWatcher->Tick(1.0f);
+				FPlatformProcess::Sleep(1.0f);
+				DirectoryWatcher->Tick(1.0f);
+
+				// close the file
+				UE_LOG(LogTestPAL, Display, TEXT("Closing FILE '%s'"), *DummyFileName);
+				delete DummyFile;
+				DummyFile = nullptr;
+				DirectoryWatcher->Tick(1.0f);
+				FPlatformProcess::Sleep(1.0f);
+				DirectoryWatcher->Tick(1.0f);
+
+				// delete file
+				UE_LOG(LogTestPAL, Display, TEXT("Deleting FILE '%s'"), *DummyFileName);
+				PlatformFile.DeleteFile(*DummyFileName);
+				DirectoryWatcher->Tick(1.0f);
+				FPlatformProcess::Sleep(1.0f);
+				DirectoryWatcher->Tick(1.0f);
+			}
+
+			UE_LOG(LogTestPAL, Display, TEXT("Deleting DIRECTORY '%s'"), *GrandChildDir);
+			verify(PlatformFile.DeleteDirectory(*GrandChildDir));
+			DirectoryWatcher->Tick(1.0f);
+			FPlatformProcess::Sleep(1.0f);
+			DirectoryWatcher->Tick(1.0f);
+		}
+
+
 		// clean up
-		DirectoryWatcher->UnregisterDirectoryChangedCallback_Handle(TestDir, DirectoryChangedHandle);
+		verify(DirectoryWatcher->UnregisterDirectoryChangedCallback_Handle(TestDir, DirectoryChangedHandle));
 		// remove dirs as well
-		PlatformFile.DeleteDirectory(*(TestDir + TEXT("/subtest")));
-		PlatformFile.DeleteDirectory(*TestDir);
+		verify(PlatformFile.DeleteDirectory(*(TestDir + TEXT("/subtest"))));
+		verify(PlatformFile.DeleteDirectory(*TestDir));
+
+		UE_LOG(LogTestPAL, Display, TEXT("End of test"));
 	}
 	else
 	{
