Linux: add missing GetOperatingSystemId() (UE-8013).

diff --git a/Engine/Source/Programs/TestPAL/Private/Main.cpp b/Engine/Source/Programs/TestPAL/Private/Main.cpp
--- a/Engine/Source/Programs/TestPAL/Private/Main.cpp
+++ b/Engine/Source/Programs/TestPAL/Private/Main.cpp
@@ -229,6 +229,28 @@ int32 ThreadSingletonTest(const TCHAR* CommandLine)
 	return 0;
 }
 
+/**
+ * Sysinfo test
+ */
+int32 SysInfoTest(const TCHAR* CommandLine)
+{
+	FPlatformMisc::SetCrashHandler(NULL);
+	FPlatformMisc::SetGracefulTerminationHandler();
+
+	GEngineLoop.PreInit(CommandLine);
+	UE_LOG(LogTestPAL, Display, TEXT("Running system info test."));
+
+	bool bIsRunningOnBattery = FPlatformMisc::IsRunningOnBattery();
+	UE_LOG(LogTestPAL, Display, TEXT("  FPlatformMisc::IsRunningOnBattery() = %s"), bIsRunningOnBattery ? TEXT("true") : TEXT("false"));
+
+	FString OSInstanceGuid = FPlatformMisc::GetOperatingSystemId();
+	UE_LOG(LogTestPAL, Display, TEXT("  FPlatformMisc::GetOperatingSystemId() = %s"), *OSInstanceGuid);
+
+	FEngineLoop::AppPreExit();
+	FEngineLoop::AppExit();
+	return 0;
+}
+
 /**
  * Selects and runs one of test cases.
  *
@@ -264,6 +286,10 @@ int32 MultiplexedMain(int32 ArgC, char* ArgV[])
 		{
 			return ThreadSingletonTest(*TestPAL::CommandLine);
 		}
+		else if (!FCStringAnsi::Strcmp(ArgV[IdxArg], ARG_SYSINFO_TEST))
+		{
+			return SysInfoTest(*TestPAL::CommandLine);
+		}
 	}
 
 	FPlatformMisc::SetCrashHandler(NULL);
@@ -279,6 +305,7 @@ int32 MultiplexedMain(int32 ArgC, char* ArgV[])
 	UE_LOG(LogTestPAL, Warning, TEXT("  %s: test message box bug (too long strings)"), ANSI_TO_TCHAR( ARG_MESSAGEBOX_TEST ));
 	UE_LOG(LogTestPAL, Warning, TEXT("  %s: test directory watcher"), ANSI_TO_TCHAR( ARG_DIRECTORY_WATCHER_TEST ));
 	UE_LOG(LogTestPAL, Warning, TEXT("  %s: test per-thread singletons"), ANSI_TO_TCHAR( ARG_THREAD_SINGLETON_TEST ));
+	UE_LOG(LogTestPAL, Warning, TEXT("  %s: test (some) system information"), ANSI_TO_TCHAR( ARG_SYSINFO_TEST ))
 	UE_LOG(LogTestPAL, Warning, TEXT(""));
 	UE_LOG(LogTestPAL, Warning, TEXT("Pass one of those to run an appropriate test."));
 
diff --git a/Engine/Source/Programs/TestPAL/Private/PrivatePCH.h b/Engine/Source/Programs/TestPAL/Private/PrivatePCH.h
--- a/Engine/Source/Programs/TestPAL/Private/PrivatePCH.h
+++ b/Engine/Source/Programs/TestPAL/Private/PrivatePCH.h
@@ -11,6 +11,7 @@
 #define ARG_MESSAGEBOX_TEST					"messagebox"
 #define ARG_DIRECTORY_WATCHER_TEST			"dirwatcher"
 #define ARG_THREAD_SINGLETON_TEST			"threadsingleton"
+#define ARG_SYSINFO_TEST					"sysinfo"
 
 // this module headers
 #include "Parent.h"
diff --git a/Engine/Source/Runtime/Core/Private/Linux/LinuxPlatformMisc.cpp b/Engine/Source/Runtime/Core/Private/Linux/LinuxPlatformMisc.cpp
--- a/Engine/Source/Runtime/Core/Private/Linux/LinuxPlatformMisc.cpp
+++ b/Engine/Source/Runtime/Core/Private/Linux/LinuxPlatformMisc.cpp
@@ -23,7 +23,8 @@ namespace PlatformMiscLimits
 {
 	enum
 	{
-		MaxUserHomeDirLength= MAX_PATH + 1
+		MaxUserHomeDirLength= MAX_PATH + 1,
+		MaxOsGuidLength = 32,
 	};
 };
 
@@ -695,3 +696,44 @@ bool FLinuxPlatformMisc::HasBeenStartedRemotely()
 
 	return bResult;
 }
+
+FString FLinuxPlatformMisc::GetOperatingSystemId()
+{
+	static bool bHasCachedResult = false;
+	static FString CachedResult;
+
+	if (!bHasCachedResult)
+	{
+		int OsGuidFile = open("/etc/machine-id", O_RDONLY);
+		if (OsGuidFile != -1)
+		{
+			char Buffer[PlatformMiscLimits::MaxOsGuidLength + 1] = {0};
+			ssize_t ReadBytes = read(OsGuidFile, Buffer, sizeof(Buffer) - 1);
+
+			if (ReadBytes > 0)
+			{
+				CachedResult = ANSI_TO_TCHAR(Buffer);
+			}
+
+			close(OsGuidFile);
+		}
+		else
+		{
+			// use old POSIX API which doesn't seem to be useful anymore
+			long HostId = gethostid();
+			if (HostId != 0)
+			{
+				CachedResult = FString::Printf(TEXT("baad051dbaad051dbaad051d%08lx"), HostId);
+			}
+		}
+
+		if (CachedResult.Len() == 0)
+		{
+			CachedResult = TEXT("dead051ddead051ddead051ddead051d");
+		}
+
+		bHasCachedResult = true;	// even if we failed to read the real one
+	}
+
+	return CachedResult;
+}
diff --git a/Engine/Source/Runtime/Core/Public/Linux/LinuxPlatformMisc.h b/Engine/Source/Runtime/Core/Public/Linux/LinuxPlatformMisc.h
--- a/Engine/Source/Runtime/Core/Public/Linux/LinuxPlatformMisc.h
+++ b/Engine/Source/Runtime/Core/Public/Linux/LinuxPlatformMisc.h
@@ -88,6 +88,7 @@ struct CORE_API FLinuxPlatformMisc : public FGenericPlatformMisc
 	static int32 NumberOfCoresIncludingHyperthreads();
 	static void LoadPreInitModules();
 	static void LoadStartupModules();
+	static FString GetOperatingSystemId();
 
 	/**
 	 * Determines the shader format for the platform
