TestPAL: add more sysinfo tests.

diff --git a/Engine/Source/Programs/TestPAL/Private/Main.cpp b/Engine/Source/Programs/TestPAL/Private/Main.cpp
--- a/Engine/Source/Programs/TestPAL/Private/Main.cpp
+++ b/Engine/Source/Programs/TestPAL/Private/Main.cpp
@@ -243,8 +243,16 @@ int32 SysInfoTest(const TCHAR* CommandLine)
 	bool bIsRunningOnBattery = FPlatformMisc::IsRunningOnBattery();
 	UE_LOG(LogTestPAL, Display, TEXT("  FPlatformMisc::IsRunningOnBattery() = %s"), bIsRunningOnBattery ? TEXT("true") : TEXT("false"));
 
+	GenericApplication * PlatformApplication = FPlatformMisc::CreateApplication();
+	checkf(PlatformApplication, TEXT("Could not create platform application!"));
+	bool bIsMouseAttached = PlatformApplication->IsMouseAttached();
+	UE_LOG(LogTestPAL, Display, TEXT("  FPlatformMisc::IsMouseAttached() = %s"), bIsMouseAttached ? TEXT("true") : TEXT("false"));
+
 	FString OSInstanceGuid = FPlatformMisc::GetOperatingSystemId();
-	UE_LOG(LogTestPAL, Display, TEXT("  FPlatformMisc::GetOperatingSystemId() = %s"), *OSInstanceGuid);
+	UE_LOG(LogTestPAL, Display, TEXT("  FPlatformMisc::GetOperatingSystemId() = '%s'"), *OSInstanceGuid);
+
+	FString MacAddress = FPlatformMisc::GetMacAddressString();
+	UE_LOG(LogTestPAL, Display, TEXT("  FPlatformMisc::GetMacAddress() = '%s'"), *MacAddress);
 
 	FPlatformMemory::DumpStats(*GLog);
 
