Do not use C functions deprecated by Microsoft.

(Edigrating CL 3501075 from Dev-Editor to Main).

diff --git a/Engine/Source/Programs/TestPAL/Private/Main.cpp b/Engine/Source/Programs/TestPAL/Private/Main.cpp
--- a/Engine/Source/Programs/TestPAL/Private/Main.cpp
+++ b/Engine/Source/Programs/TestPAL/Private/Main.cpp
@@ -620,8 +620,13 @@ int32 MallocThreadingTest(const TCHAR* CommandLine)
  */
 void ReplayMallocFile(const FString& ReplayFileName, uint64 OperationToStopAfter, bool bSuppressErrors)
 {
+#if PLATFORM_WINDOWS
+	FILE* ReplayFile = nullptr;
+	if (fopen_s(&ReplayFile, TCHAR_TO_UTF8(*ReplayFileName), "rb") != 0 || ReplayFile == nullptr)
+#else
 	FILE* ReplayFile = fopen(TCHAR_TO_UTF8(*ReplayFileName), "rb");
 	if (ReplayFile == nullptr)
+#endif // PLATFORM_WINDOWS
 	{
 		return;
 	}
@@ -645,7 +650,11 @@ void ReplayMallocFile(const FString& ReplayFileName, uint64 OperationToStopAfter
 		char OpBuffer[128] = {0};
 		uint64 PtrOut, PtrIn, Size, Ordinal;
 		uint32 Alignment;
+#if PLATFORM_WINDOWS
+		if (fscanf_s(ReplayFile, "%s %llu %llu %llu %u\t# %llu\n", OpBuffer, static_cast<unsigned int>(sizeof(OpBuffer)), &PtrOut, &PtrIn, &Size, &Alignment, &Ordinal) != 6)
+#else
 		if (fscanf(ReplayFile, "%s %llu %llu %llu %u\t# %llu\n", OpBuffer, &PtrOut, &PtrIn, &Size, &Alignment, &Ordinal) != 6)
+#endif // PLATFORM_WINDOWS
 		{
 			UE_LOG(LogTestPAL, Display, TEXT("Hit end of the replay file on %llu-th operation."), OperationNumber);
 			break;
