Set current working dir

diff --git a/Engine/Source/Programs/TestPAL/Private/Main.cpp b/Engine/Source/Programs/TestPAL/Private/Main.cpp
--- a/Engine/Source/Programs/TestPAL/Private/Main.cpp
+++ b/Engine/Source/Programs/TestPAL/Private/Main.cpp
@@ -276,10 +276,13 @@ int32 SysInfoTest(const TCHAR* CommandLine)
 	UE_LOG(LogTestPAL, Display, TEXT("  FPlatformMisc::GetOperatingSystemId() = '%s'"), *OSInstanceGuid);
 
 	FString UserDir = FPlatformProcess::UserDir();
-	UE_LOG(LogTestPAL, Display, TEXT("  FPlatformMisc::UserDir() = '%s'"), *UserDir);
+	UE_LOG(LogTestPAL, Display, TEXT("  FPlatformProcess::UserDir() = '%s'"), *UserDir);
 
 	FString ApplicationSettingsDir = FPlatformProcess::ApplicationSettingsDir();
-	UE_LOG(LogTestPAL, Display, TEXT("  FPlatformMisc::ApplicationSettingsDir() = '%s'"), *ApplicationSettingsDir);
+	UE_LOG(LogTestPAL, Display, TEXT("  FPlatformProcess::ApplicationSettingsDir() = '%s'"), *ApplicationSettingsDir);
+
+	FString ApplicationCurrentWorkingDir = FPlatformProcess::GetCurrentWorkingDirectory();
+	UE_LOG(LogTestPAL, Display, TEXT("  FPlatformProcess::GetCurrentWorkingDirectory() = '%s'"), *ApplicationCurrentWorkingDir);
 
 	FPlatformMemory::DumpStats(*GLog);
 
diff --git a/Engine/Source/Runtime/Core/Private/Unix/UnixPlatformProcess.cpp b/Engine/Source/Runtime/Core/Private/Unix/UnixPlatformProcess.cpp
--- a/Engine/Source/Runtime/Core/Private/Unix/UnixPlatformProcess.cpp
+++ b/Engine/Source/Runtime/Core/Private/Unix/UnixPlatformProcess.cpp
@@ -1419,6 +1419,12 @@ uint32 FUnixPlatformProcess::GetCurrentProcessId()
 	return getpid();
 }
 
+void FUnixPlatformProcess::SetCurrentWorkingDirectoryToBaseDir()
+{
+	FPlatformMisc::CacheLaunchDir();
+	chdir(TCHAR_TO_ANSI(FPlatformProcess::BaseDir()));
+}
+
 FString FUnixPlatformProcess::GetCurrentWorkingDirectory()
 {
 	// get the current directory
diff --git a/Engine/Source/Runtime/Core/Public/Unix/UnixPlatformProcess.h b/Engine/Source/Runtime/Core/Public/Unix/UnixPlatformProcess.h
--- a/Engine/Source/Runtime/Core/Public/Unix/UnixPlatformProcess.h
+++ b/Engine/Source/Runtime/Core/Public/Unix/UnixPlatformProcess.h
@@ -246,6 +246,7 @@ struct CORE_API FUnixPlatformProcess : public FGenericPlatformProcess
 	static const TCHAR* UserDir();
 	static const TCHAR* UserSettingsDir();
 	static const TCHAR* ApplicationSettingsDir();
+	static void SetCurrentWorkingDirectoryToBaseDir();
 	static FString GetCurrentWorkingDirectory();
 	static FString GenerateApplicationPath(const FString& AppName, EBuildConfigurations::Type BuildConfiguration);
 	static FString GetApplicationName( uint32 ProcessId );
