TestPAL: Add cases for testing inline callstacks

diff --git a/Engine/Source/Programs/TestPAL/Private/Main.cpp b/Engine/Source/Programs/TestPAL/Private/Main.cpp
--- a/Engine/Source/Programs/TestPAL/Private/Main.cpp
+++ b/Engine/Source/Programs/TestPAL/Private/Main.cpp
@@ -31,6 +31,7 @@ IMPLEMENT_APPLICATION(TestPAL, "TestPAL");
 #define ARG_MALLOC_THREADING_TEST			"mallocthreadtest"
 #define ARG_MALLOC_REPLAY					"mallocreplay"
 #define ARG_THREAD_PRIO_TEST				"threadpriotest"
+#define ARG_INLINE_CALLSTACK_TEST			"inline"
 
 namespace TestPAL
 {
@@ -1058,6 +1059,113 @@ int32 ThreadPriorityTest(const TCHAR* CommandLine)
 	return 0;
 }
 
+// inlined/non-inlined functions for testing
+namespace
+{
+	void FORCENOINLINE LexicalBlock()
+	{
+		{
+			{
+				ensure(false);
+			}
+		}
+	}
+
+	void FORCENOINLINE LabelSwitch()
+	{
+		switch (1)
+		{
+			case 1:
+				ensure(false);
+				break;
+			default:
+				break;
+		}
+	}
+
+	void FORCENOINLINE LabelGoto()
+	{
+		goto end;
+		ensure(false); // skips all of these. Just want to create a bunch of inline statements in this single non inlined one
+		ensure(false);
+		ensure(false);
+		ensure(false);
+		ensure(false);
+		ensure(false);
+end:
+		ensure(false);
+	}
+
+	void FORCEINLINE inline_three_ensures()
+	{
+		ensure(false);
+	}
+
+	void FORCEINLINE inline_two_calls_inline_three()
+	{
+		inline_three_ensures();
+	}
+
+	void FORCEINLINE inline_one_calls_inline_two()
+	{
+		inline_two_calls_inline_three();
+	}
+
+	void FORCENOINLINE MultipleInlineDeep()
+	{
+		inline_one_calls_inline_two();
+	}
+
+	void FORCEINLINE inline_crash()
+	{
+		*(int*)0x1 = 0x0;
+	}
+
+	void FORCENOINLINE no_inline_to_inline_crash()
+	{
+		inline_crash();
+	}
+}
+
+int32 InlineCallstacksTest(const TCHAR* CommandLine)
+{
+	FPlatformMisc::SetCrashHandler(NULL);
+	FPlatformMisc::SetGracefulTerminationHandler();
+
+	GEngineLoop.PreInit(CommandLine);
+
+	/*  Three unique cases that we can only test two:
+	 *
+	 *  LexicalBlocks (ie. scope/if/for/while blocks
+	 *  Labels (ie. goto/switch statements)
+	 *  Try/Catch (cannot use when execptions are disabled, but we *should* handle this case)
+	 *
+	 *  As well as testing, multiple deep inlining. ie. Calling multiple inline functions and still
+	 *  seeing the call site
+	 *
+	 *  Just make sure that the functions used + their callsite into ensure(false) are correctly reported in
+	 *  the callstack
+	 */
+
+	UE_LOG(LogTestPAL, Warning, TEXT("\n*** Lexical Block ***"));
+	LexicalBlock();
+	UE_LOG(LogTestPAL, Warning, TEXT("\n*** Label Switch ***"));
+	LabelSwitch();
+
+	UE_LOG(LogTestPAL, Warning, TEXT("\n*** Label Goto ***"));
+	LabelGoto();
+
+	UE_LOG(LogTestPAL, Warning, TEXT("\n*** Multiple Inlined ***"));
+	MultipleInlineDeep();
+
+	// Should always be the last case, as this crashes
+	UE_LOG(LogTestPAL, Warning, TEXT("\n*** Array delegate to crash***"));
+	TArray<TFunction<void()>> a;
+	a.Push([] { no_inline_to_inline_crash(); });
+	a[0]();
+
+	return 0;
+}
 
 /**
  * Selects and runs one of test cases.
@@ -1126,6 +1234,10 @@ int32 MultiplexedMain(int32 ArgC, char* ArgV[])
 		{
 			return ThreadPriorityTest(*TestPAL::CommandLine);
 		}
+		else if (!FCStringAnsi::Strcmp(ArgV[IdxArg], ARG_INLINE_CALLSTACK_TEST))
+		{
+			return InlineCallstacksTest(*TestPAL::CommandLine);
+		}
 	}
 
 	FPlatformMisc::SetCrashHandler(NULL);
@@ -1149,6 +1261,7 @@ int32 MultiplexedMain(int32 ArgC, char* ArgV[])
 	UE_LOG(LogTestPAL, Warning, TEXT("  %s: test malloc for thread-safety and performance."), UTF8_TO_TCHAR(ARG_MALLOC_THREADING_TEST));
 	UE_LOG(LogTestPAL, Warning, TEXT("  %s: test by replaying a saved malloc history saved by -mallocsavereplay. Possible options: -replayfile=File, -stopafter=N (operation), -suppresserrors"), UTF8_TO_TCHAR(ARG_MALLOC_REPLAY));
 	UE_LOG(LogTestPAL, Warning, TEXT("  %s: test thread priorities."), UTF8_TO_TCHAR(ARG_THREAD_PRIO_TEST));
+	UE_LOG(LogTestPAL, Warning, TEXT("  %s: test inline callstacks through ensures and a final crash."), UTF8_TO_TCHAR(ARG_INLINE_CALLSTACK_TEST));
 	UE_LOG(LogTestPAL, Warning, TEXT(""));
 	UE_LOG(LogTestPAL, Warning, TEXT("Pass one of those to run an appropriate test."));
 
