Implement better Linux CreateGuid using SYS_getrandom (if available)

Getrandom syscall was introduced in version Linux kernel v3.17.

Can be disabled with -norandomguids command line option.

On my Ubuntu 18.04 test machine[**], single threaded perf is slightly slower,
but performance with multiple threads improves:

  With Threads: 1, Calls: 1000000
    FGenericPlatformMisc::CreateGuid: 0.316885 seconds
    FPlatformMisc::CreateGuid:        0.544386 seconds

  With Threads: 8, Calls: 1000000
    FGenericPlatformMisc::CreateGuid: 5.867290 seconds
    FPlatformMisc::CreateGuid:        3.035365 seconds

  With Threads: 16, Calls: 1000000
    FGenericPlatformMisc::CreateGuid: 10.914646 seconds
    FPlatformMisc::CreateGuid:         5.897632 seconds

[**]
  4.15.0-45-generic
  Intel(R) Xeon(R) Gold 6134 CPU @ 3.20GHz)

Guids with getrandom() appear to be significantly more random:
 0 GuidGeneric:{4240FA8F-08D6-974C-F710-112EB0D1AB97} GuidPlatform:{F49E1FFA-67A1-4702-B065-6E1B6385BC5E}
 1 GuidGeneric:{4241444E-08D6-974C-F710-1296B0D1ABB9} GuidPlatform:{77D96835-24F5-4A27-B9E7-70985316285F}
 2 GuidGeneric:{4242B446-08D6-974C-F710-13CCB0D1ABD9} GuidPlatform:{80C46DF2-6DD1-4B27-B1BD-CC279B5874A0}
 3 GuidGeneric:{424386A9-08D6-974C-F710-1502B0D1ABF8} GuidPlatform:{5DA06154-78A5-47BB-A9A6-75B220950F31}
 4 GuidGeneric:{424436F2-08D6-974C-F710-1638B0D1AC17} GuidPlatform:{880670CF-0ABF-491F-80A4-91C28B6D5B2F}

diff --git a/Engine/Source/Programs/TestPAL/Private/Main.cpp b/Engine/Source/Programs/TestPAL/Private/Main.cpp
--- a/Engine/Source/Programs/TestPAL/Private/Main.cpp
+++ b/Engine/Source/Programs/TestPAL/Private/Main.cpp
@@ -3,6 +3,7 @@
 #include "CoreMinimal.h"
 #include "TestPALLog.h"
 #include "Parent.h"
+#include "Misc/Guid.h"
 #include "Stats/StatsMisc.h"
 #include "HAL/RunnableThread.h"
 #include "HAL/PlatformApplicationMisc.h"
@@ -34,6 +35,7 @@ IMPLEMENT_APPLICATION(TestPAL, "TestPAL");
 #define ARG_THREAD_PRIO_TEST				"threadpriotest"
 #define ARG_INLINE_CALLSTACK_TEST			"inline"
 #define ARG_STRINGS_ALLOCATION_TEST			"stringsallocation"
+#define ARG_CREATEGUID_TEST					"createguid"
 
 namespace TestPAL
 {
@@ -680,6 +682,135 @@ int32 StringsAllocationTest(const TCHAR* CommandLine)
 	return 0;
 }
 
+/**
+ * Thread runnable
+ * bUseGenericMisc Whether to use FGenericPlatformMisc CreateGuid
+ */
+template<bool bUseGenericMisc>
+struct FCreateGuidThread : public FRunnable
+{
+	/** Number of CreateGuid calls to make. */
+	int32	NumCalls;
+
+	FCreateGuidThread(int32 InNumCalls)
+		:	FRunnable()
+		,	NumCalls(InNumCalls)
+	{
+	}
+
+	virtual uint32 Run()
+	{
+		FGuid Result;
+
+		for (int32 IdxRun = 0; IdxRun < NumCalls; ++IdxRun)
+		{
+			if (bUseGenericMisc)
+			{
+				FGenericPlatformMisc::CreateGuid(Result);
+			}
+			else
+			{
+				FPlatformMisc::CreateGuid(Result);
+			}
+		}
+
+		return 0;
+	}
+};
+
+
+/**
+ * CreateGuid test
+ */
+int32 CreateGuidTest(const TCHAR* CommandLine)
+{
+	FPlatformMisc::SetCrashHandler(NULL);
+	FPlatformMisc::SetGracefulTerminationHandler();
+
+	GEngineLoop.PreInit(CommandLine);
+	UE_LOG(LogTestPAL, Display, TEXT("Running CreateGuid test."));
+
+	TArray<FRunnable*> RunnableArray;
+	TArray<FRunnableThread*> ThreadArray;
+
+	UE_LOG(LogTestPAL, Display, TEXT("Accepted options:"));
+	UE_LOG(LogTestPAL, Display, TEXT("    -numthreads=N"));
+	UE_LOG(LogTestPAL, Display, TEXT("    -numcalls=N (how many times each thread will call CreateGuid)"));
+	UE_LOG(LogTestPAL, Display, TEXT("    -norandomguids (disable SYS_getrandom syscall)"));
+
+	int32 NumTestThreads = 4, NumCalls = 1000000;
+	if (FParse::Value(CommandLine, TEXT("numthreads="), NumTestThreads))
+	{
+		NumTestThreads = FMath::Max(1, NumTestThreads);
+	}
+	if (FParse::Value(CommandLine, TEXT("numcalls="), NumCalls))
+	{
+		NumCalls = FMath::Max(1, NumCalls);
+	}
+
+	// start all threads
+	for (int Idx = 0; Idx < 2; ++Idx )
+	{
+		bool bUseGenericMisc = (Idx == 0);
+		double WallTimeDuration = FPlatformTime::Seconds();
+
+		for (int IdxThread = 0; IdxThread < NumTestThreads; ++IdxThread)
+		{
+			RunnableArray.Add(bUseGenericMisc ?
+				static_cast<FRunnable*>(new FCreateGuidThread<true>(NumCalls)) :
+				static_cast<FRunnable*>(new FCreateGuidThread<false>(NumCalls)) );
+
+			ThreadArray.Add( FRunnableThread::Create(RunnableArray[IdxThread],
+				*FString::Printf(TEXT("GuidTest%d"), IdxThread)) );
+		}
+
+		GLog->FlushThreadedLogs();
+		GLog->Flush();
+
+		// join all threads
+		for (int IdxThread = 0; IdxThread < NumTestThreads; ++IdxThread)
+		{
+			ThreadArray[IdxThread]->WaitForCompletion();
+
+			delete ThreadArray[IdxThread];
+			ThreadArray[IdxThread] = nullptr;
+
+			delete RunnableArray[IdxThread];
+			RunnableArray[IdxThread] = nullptr;
+		}
+
+		WallTimeDuration = FPlatformTime::Seconds() - WallTimeDuration;
+
+		UE_LOG(LogTestPAL, Display, TEXT("--- Results for %s ---"),
+			bUseGenericMisc ? TEXT("FGenericPlatformMisc::CreateGuid") : TEXT("FPlatformMisc::CreateGuid"));
+		UE_LOG(LogTestPAL, Display, TEXT("Total wall time: %f seconds, Threads: %d, Calls: %d"),
+				WallTimeDuration, NumTestThreads, NumCalls);
+
+		ThreadArray.Empty();
+		RunnableArray.Empty();
+	}
+
+	// Spew out a small sample of GUIDs to visually check we haven't horribly broken something
+	UE_LOG(LogTestPAL, Display, TEXT("--- CreateGuid Samples ---"));
+
+	for (int Idx = 0; Idx < 20; ++Idx)
+	{
+		FGuid GuidPlatform;
+		FGuid GuidGeneric;
+
+		FPlatformMisc::CreateGuid(GuidPlatform);
+		FGenericPlatformMisc::CreateGuid(GuidGeneric);
+
+		UE_LOG(LogTestPAL, Display, TEXT("%3d GuidGeneric:%s GuidPlatform:%s"), Idx,
+			*GuidGeneric.ToString(EGuidFormats::DigitsWithHyphensInBraces),
+			*GuidPlatform.ToString(EGuidFormats::DigitsWithHyphensInBraces));
+	}
+
+	// GMalloc = OldGMalloc;
+	FEngineLoop::AppExit();
+	return 0;
+}
+
 
 /** An ugly way to pass a parameters to FRunnable; shouldn't matter for this test code. */
 int32 GMallocTestNumRuns = 500;
@@ -1418,6 +1549,10 @@ int32 MultiplexedMain(int32 ArgC, char* ArgV[])
 		{
 			return StringsAllocationTest(*TestPAL::CommandLine);
 		}
+		else if (!FCStringAnsi::Strcmp(ArgV[IdxArg], ARG_CREATEGUID_TEST))
+		{
+			return CreateGuidTest(*TestPAL::CommandLine);
+		}
 	}
 
 	FPlatformMisc::SetCrashHandler(NULL);
@@ -1443,6 +1578,8 @@ int32 MultiplexedMain(int32 ArgC, char* ArgV[])
 	UE_LOG(LogTestPAL, Warning, TEXT("  %s: test by replaying a saved malloc history saved by -mallocsavereplay. Possible options: -replayfile=File, -stopafter=N (operation), -suppresserrors"), UTF8_TO_TCHAR(ARG_MALLOC_REPLAY));
 	UE_LOG(LogTestPAL, Warning, TEXT("  %s: test thread priorities."), UTF8_TO_TCHAR(ARG_THREAD_PRIO_TEST));
 	UE_LOG(LogTestPAL, Warning, TEXT("  %s: test inline callstacks through ensures and a final crash."), UTF8_TO_TCHAR(ARG_INLINE_CALLSTACK_TEST));
+	UE_LOG(LogTestPAL, Warning, TEXT("  %s: test string allocations."), UTF8_TO_TCHAR(ARG_STRINGS_ALLOCATION_TEST));
+	UE_LOG(LogTestPAL, Warning, TEXT("  %s: test CreateGuid."), UTF8_TO_TCHAR(ARG_CREATEGUID_TEST));
 	UE_LOG(LogTestPAL, Warning, TEXT(""));
 	UE_LOG(LogTestPAL, Warning, TEXT("Pass one of those to run an appropriate test."));
 
diff --git a/Engine/Source/Runtime/Core/Private/Unix/UnixPlatformMisc.cpp b/Engine/Source/Runtime/Core/Private/Unix/UnixPlatformMisc.cpp
--- a/Engine/Source/Runtime/Core/Private/Unix/UnixPlatformMisc.cpp
+++ b/Engine/Source/Runtime/Core/Private/Unix/UnixPlatformMisc.cpp
@@ -45,6 +45,8 @@
 
 extern bool GInitializedSDL;
 
+static int SysGetRandomSupported = -1;
+
 namespace PlatformMiscLimits
 {
 	enum
@@ -233,6 +235,12 @@ void FUnixPlatformMisc::PlatformInit()
 		// print output immediately
 		setvbuf(stdout, NULL, _IONBF, 0);
 	}
+
+	if (FParse::Param(FCommandLine::Get(), TEXT("norandomguids")))
+	{
+		// If "-norandomguids" specified, don't use SYS_getrandom syscall
+		SysGetRandomSupported = 0;
+	}
 }
 
 extern void CORE_API UnixPlatformStackWalk_UnloadPreloadedModuleSymbol();
@@ -904,6 +912,85 @@ bool FUnixPlatformMisc::IsRunningOnBattery()
 	return bIsOnBattery;
 }
 
+#if PLATFORM_UNIX && defined(_GNU_SOURCE)
+
+#include <sys/syscall.h>
+
+// http://man7.org/linux/man-pages/man2/getrandom.2.html
+// getrandom() was introduced in version 3.17 of the Linux kernel
+//   and glibc version 2.25.
+
+// Check known platforms if SYS_getrandom isn't defined
+#if !defined(SYS_getrandom)
+	#if PLATFORM_CPU_X86_FAMILY && PLATFORM_64BITS
+		#define SYS_getrandom 318
+	#elif PLATFORM_CPU_X86_FAMILY && !PLATFORM_64BITS
+		#define SYS_getrandom 355
+	#elif PLATFORM_CPU_ARM_FAMILY && PLATFORM_64BITS
+		#define SYS_getrandom 278
+	#elif PLATFORM_CPU_ARM_FAMILY && !PLATFORM_64BITS
+		#define SYS_getrandom 384
+	#endif
+#endif // !defined(SYS_getrandom)
+
+#endif // PLATFORM_UNIX && _GNU_SOURCE
+
+namespace
+{
+#if defined(SYS_getrandom)
+
+#if !defined(GRND_NONBLOCK)
+	#define GRND_NONBLOCK 0x0001
+#endif
+	
+	int SysGetRandom(void *buf, size_t buflen)
+	{
+		if (SysGetRandomSupported < 0)
+		{
+			int Ret = syscall(SYS_getrandom, buf, buflen, GRND_NONBLOCK);
+	
+			// If -1 is returned with ENOSYS, kernel doesn't support getrandom
+			SysGetRandomSupported = ((Ret == -1) && (errno == ENOSYS)) ? 0 : 1;
+		}
+	
+		return SysGetRandomSupported ?
+			syscall(SYS_getrandom, buf, buflen, GRND_NONBLOCK) : -1;
+	}
+	
+#else
+
+	int SysGetRandom(void *buf, size_t buflen)
+	{
+		return -1;
+	}
+	
+#endif // !SYS_getrandom
+}
+
+// If we fail to create a Guid with urandom fallback to the generic platform.
+// This maybe need to be tweaked for Servers and hard fail here
+void FUnixPlatformMisc::CreateGuid(FGuid& Result)
+{
+	int BytesRead = SysGetRandom(&Result, sizeof(Result));
+
+	if (BytesRead == sizeof(Result))
+	{
+		// https://tools.ietf.org/html/rfc4122#section-4.4
+		// https://en.wikipedia.org/wiki/Universally_unique_identifier
+		//
+		// The 4 bits of digit M indicate the UUID version, and the 1–3
+		//   most significant bits of digit N indicate the UUID variant.
+		// xxxxxxxx-xxxx-Mxxx-Nxxx-xxxxxxxxxxxx
+		Result[1] = (Result[1] & 0xffff0fff) | 0x00004000; // version 4
+		Result[2] = (Result[2] & 0x3fffffff) | 0x80000000; // variant 1
+	}
+	else
+	{
+		// Fall back to generic CreateGuid
+		FGenericPlatformMisc::CreateGuid(Result);
+	}
+}
+
 #if STATS || ENABLE_STATNAMEDEVENTS
 void FUnixPlatformMisc::BeginNamedEventFrame()
 {
diff --git a/Engine/Source/Runtime/Core/Public/Unix/UnixPlatformMisc.h b/Engine/Source/Runtime/Core/Public/Unix/UnixPlatformMisc.h
--- a/Engine/Source/Runtime/Core/Public/Unix/UnixPlatformMisc.h
+++ b/Engine/Source/Runtime/Core/Public/Unix/UnixPlatformMisc.h
@@ -133,6 +133,8 @@ struct CORE_API FUnixPlatformMisc : public FGenericPlatformMisc
 	static bool HasOverriddenReturnCode(uint8 * OverriddenReturnCodeToUsePtr);
 	static FString GetOSVersion();
 
+	static void CreateGuid(FGuid& Result);
+
 #if STATS || ENABLE_STATNAMEDEVENTS
 	static void BeginNamedEventFrame();
 	static void BeginNamedEvent(const struct FColor& Color, const TCHAR* Text);
