Removing reference to hardcoded Steam path from dll loading test in pal app.

diff --git a/Engine/Source/Programs/TestPAL/Private/Main.cpp b/Engine/Source/Programs/TestPAL/Private/Main.cpp
--- a/Engine/Source/Programs/TestPAL/Private/Main.cpp
+++ b/Engine/Source/Programs/TestPAL/Private/Main.cpp
@@ -500,12 +500,31 @@ int32 DynamicLibraryTest(const TCHAR* CommandLine)
 
 	if (PLATFORM_LINUX)
 	{
-		RootSteamPath = FPaths::EngineDir() / FString(TEXT("Binaries/ThirdParty/Steamworks/Steamv139/x86_64-unknown-linux-gnu/"));
+		RootSteamPath = FPaths::EngineDir() / FString(TEXT("Binaries/ThirdParty/Steamworks/*"));
+
+		IFileManager& PlatformFileManager = IFileManager::Get();
+		TArray<FString> FoundDirectories;
+		PlatformFileManager.FindFiles(FoundDirectories, *RootSteamPath, false, true);
+
+		// Just use the first directory we find, this test does not have to be very sophisticated.
+		if (FoundDirectories.Num() > 0)
+		{
+			// This only gives us directories, so remove the wildcard from our initial search
+			RootSteamPath.RemoveFromEnd(TEXT("*"));
+			// And append the directory name we found.
+			RootSteamPath += FoundDirectories[0];
+		}
+		else
+		{
+			UE_LOG(LogTestPAL, Fatal, TEXT("Could not find any steam versions."));
+		}
+
+		RootSteamPath += TEXT("/x86_64-unknown-linux-gnu/");
 		LibraryName = TEXT("libsteam_api.so");
 	}
 	else
 	{
-		UE_LOG(LogTestPAL, Fatal, TEXT("This test is not implemented for this platform."))
+		UE_LOG(LogTestPAL, Fatal, TEXT("This test is not implemented for this platform."));
 	}
 
 	FPlatformProcess::PushDllDirectory(*RootSteamPath);
@@ -520,7 +539,7 @@ int32 DynamicLibraryTest(const TCHAR* CommandLine)
 
 		if (SteamDLLHandle == nullptr)
 		{
-			UE_LOG(LogTestPAL, Fatal, TEXT("Could not load Steam library!"))
+			UE_LOG(LogTestPAL, Fatal, TEXT("Could not load Steam library!"));
 		}
 	}
 
