Fix compiler errors for TestPAL

diff --git a/Engine/Source/Programs/TestPAL/Private/Main.cpp b/Engine/Source/Programs/TestPAL/Private/Main.cpp
--- a/Engine/Source/Programs/TestPAL/Private/Main.cpp
+++ b/Engine/Source/Programs/TestPAL/Private/Main.cpp
@@ -1483,17 +1483,17 @@ struct FThreadStackTest : public FRunnable
 		// If we dont have a valid ThreadId lets use the threads id as default
 		if (ThreadId == ~0)
 		{
-			ThreadId = syscall(SYS_gettid);
+			ThreadId = FPlatformTLS::GetCurrentThreadId();
 		}
 
 		// Hopefully this wont be to much extra stack space for a testing thread
-		SIZE_T StackTraceSize = 65536;
+		const SIZE_T StackTraceSize = 65536;
 		ANSICHAR StackTrace[StackTraceSize] = {0};
 		FPlatformStackWalk::ThreadStackWalkAndDump(StackTrace, StackTraceSize, 0, ThreadId);
 
 		UE_LOG(LogTestPAL, Warning, TEXT("***** ThreadStackWalkAndDump for ThreadId(%lu) ******\n%s"), ThreadId, ANSI_TO_TCHAR(StackTrace));
 
-		SIZE_T BackTraceSize = 100;
+		const SIZE_T BackTraceSize = 100;
 		uint64 BackTrace[BackTraceSize];
 		int32 BackTraceCount = FPlatformStackWalk::CaptureThreadStackBackTrace(ThreadId, BackTrace, BackTraceSize);
 
