TestPAL: Added an -ensure param to the crash command to test a failed ensure

#rnx

diff --git a/Engine/Source/Programs/TestPAL/Private/Main.cpp b/Engine/Source/Programs/TestPAL/Private/Main.cpp
--- a/Engine/Source/Programs/TestPAL/Private/Main.cpp
+++ b/Engine/Source/Programs/TestPAL/Private/Main.cpp
@@ -345,6 +345,10 @@ int32 CrashTest(const TCHAR* CommandLine)
 	{
 		checkf(false, TEXT("  checkf!"));
 	}
+	else if (FParse::Param(CommandLine, TEXT("ensure")))
+	{
+		ensureMsgf(false, TEXT("  ensureMsgf!"));
+	}
 	else if (FParse::Param(CommandLine, TEXT("unaligned")))
 	{
 		FQuat Quat[2];
