Fix unaceptable words usage in program

diff --git a/Engine/Source/Programs/TestPAL/Private/Main.cpp b/Engine/Source/Programs/TestPAL/Private/Main.cpp
--- a/Engine/Source/Programs/TestPAL/Private/Main.cpp
+++ b/Engine/Source/Programs/TestPAL/Private/Main.cpp
@@ -94,7 +94,7 @@ int32 ProcRunAsParent(const TCHAR* CommandLine)
 	GEngineLoop.PreInit(CommandLine);
 	UE_LOG(LogTestPAL, Display, TEXT("Running proc test as parent."));
 
-	// Run slave instance continuously
+	// Run child instance continuously
 	int NumChildrenToSpawn = 255, MaxAtOnce = 5;
 	FParent Parent(NumChildrenToSpawn, MaxAtOnce);
 
