Add 'cmdline' test to Linux TestPAL to aid in catching differences between how other platforms parse command line arguments

#rnx

diff --git a/Engine/Source/Programs/TestPAL/Private/Main.cpp b/Engine/Source/Programs/TestPAL/Private/Main.cpp
--- a/Engine/Source/Programs/TestPAL/Private/Main.cpp
+++ b/Engine/Source/Programs/TestPAL/Private/Main.cpp
@@ -42,6 +42,7 @@ IMPLEMENT_APPLICATION(TestPAL, "TestPAL");
 #define ARG_THREADSTACK_TEST				"threadstack"
 #define ARG_EXEC_PROCESS_TEST				"exec-process"
 #define ARG_FORK_TEST						"fork"
+#define ARG_CMDLINE_PARSE_TEST				"cmdline"
 
 namespace TestPAL
 {
@@ -1684,6 +1685,56 @@ int32 ExecProcessTest(const TCHAR *CommandLine)
 	return 0;
 }
 
+int32 CmdlineParseTest(const TCHAR *CommandLine)
+{
+	int32 NotOk = 0;	// zero when everything has passed so far
+	FPlatformMisc::SetCrashHandler(NULL);
+	FPlatformMisc::SetGracefulTerminationHandler();
+
+	GEngineLoop.PreInit(CommandLine);
+	
+	void* ReadPipe = nullptr;
+	void* WritePipe = nullptr;
+	verify(FPlatformProcess::CreatePipe(ReadPipe, WritePipe));
+
+	uint32 ProcessID;
+
+	// key is the command line to test, value is the expected parse result when run through 'echo'
+	TArray<TTuple<FString, FString>> tests = 
+	{ 
+		{ TEXT("foo"), TEXT("foo") },
+		{ TEXT("\"foo bar\""), TEXT("foo bar") },
+		{ TEXT("\"\"foo bar\"\""), TEXT("\"foo bar\"") },
+		{ TEXT("\"foo"), TEXT("") },
+		{ TEXT("\"\"foo"), TEXT("") },
+		{ TEXT("-logpath=\"path with space\""), TEXT("-logpath=path with space") },
+		{ TEXT("-logpath=\"\"double quoted path with space\"\""), TEXT("-logpath=\"double quoted path with space\"") },
+		{ TEXT("bar \"\"foo blarg"), TEXT("bar") },
+	};
+
+	for(int i=0;i<tests.Num();i++)
+	{
+		FProcHandle ProcessHandle = FPlatformProcess::CreateProc(TEXT("/usr/bin/echo"), *tests[i].Key, false, false, false, &ProcessID, 0, nullptr, WritePipe);
+		while (FPlatformProcess::IsProcRunning(ProcessHandle));
+		FString Result = FPlatformProcess::ReadPipe(ReadPipe).TrimEnd();
+		NotOk = Result.Compare(*tests[i].Value);
+		if(NotOk)
+		{
+			UE_LOG(LogTestPAL, Display, TEXT("CmdLine test failed: '%s' != '%s'\n"), *Result, *tests[i].Value);
+			break;
+		}
+	}
+	
+	FPlatformProcess::ClosePipe(ReadPipe, WritePipe);
+
+	UE_LOG(LogTestPAL, Display, TEXT("Parse command line test %s"), (NotOk)?TEXT("failed!"):TEXT("succeeded!"));
+
+	FEngineLoop::AppPreExit();
+	FEngineLoop::AppExit();
+
+	return NotOk;
+}
+
 /**
  * Selects and runs one of test cases.
  *
@@ -1781,6 +1832,10 @@ int32 MultiplexedMain(int32 ArgC, char* ArgV[])
 		{
 			return ForkTest(*TestPAL::CommandLine);
 		}
+		else if (PLATFORM_LINUX && !FCStringAnsi::Strcmp(ArgV[IdxArg], ARG_CMDLINE_PARSE_TEST))
+		{
+			return CmdlineParseTest(*TestPAL::CommandLine);
+		}
 	}
 
 	FPlatformMisc::SetCrashHandler(NULL);
@@ -1814,6 +1869,7 @@ int32 MultiplexedMain(int32 ArgC, char* ArgV[])
 	if (PLATFORM_LINUX)
 	{
 		UE_LOG(LogTestPAL, Warning, TEXT("  %s: test WaitAndFork"), UTF8_TO_TCHAR(ARG_FORK_TEST));
+		UE_LOG(LogTestPAL, Warning, TEXT("  %s: test CmdlineParse"), UTF8_TO_TCHAR(ARG_CMDLINE_PARSE_TEST));
 	}
 
 	UE_LOG(LogTestPAL, Warning, TEXT(""));
