Fix C4702 unreachable code
Fix warning C6237: (<zero> && <expression>) is always zero.  <expression> is never evaluated and might have side effects.
Fix warning C6262: Function uses '8000044' bytes of stack.  Consider moving some data to heap.

diff --git a/Engine/Source/Programs/TestPAL/Private/Main.cpp b/Engine/Source/Programs/TestPAL/Private/Main.cpp
--- a/Engine/Source/Programs/TestPAL/Private/Main.cpp
+++ b/Engine/Source/Programs/TestPAL/Private/Main.cpp
@@ -712,7 +712,7 @@ int32 StringsAllocationTest(const TCHAR* CommandLine)
 	const int32 NumOfStrings = 1000000;
 	const TCHAR* SampleText = TEXT("Lorem ipsum dolor sit amet");
 
-	FString* Strings[NumOfStrings];
+	FString** Strings = new FString*[NumOfStrings];
 
 	UE_LOG(LogTestPAL, Display, TEXT("Allocating %u strings '%s'"), NumOfStrings, SampleText);
 
@@ -730,6 +730,7 @@ int32 StringsAllocationTest(const TCHAR* CommandLine)
 	{
 		delete Strings[i];
 	}
+	delete [] Strings;
 
 	// GMalloc = OldGMalloc;
 	FEngineLoop::AppExit();
@@ -1438,6 +1439,8 @@ namespace
 		}
 	}
 
+#pragma warning(push)
+#pragma warning(disable:4702)
 	void FORCENOINLINE LabelGoto()
 	{
 		goto end;
@@ -1450,6 +1453,7 @@ namespace
 end:
 		ensure(false);
 	}
+#pragma warning(pop)
 
 	void FORCEINLINE inline_three_ensures()
 	{
@@ -1832,10 +1836,12 @@ int32 MultiplexedMain(int32 ArgC, char* ArgV[])
 		{
 			return ForkTest(*TestPAL::CommandLine);
 		}
-		else if (PLATFORM_LINUX && !FCStringAnsi::Strcmp(ArgV[IdxArg], ARG_CMDLINE_PARSE_TEST))
+#if PLATFORM_LINUX
+		else if (!FCStringAnsi::Strcmp(ArgV[IdxArg], ARG_CMDLINE_PARSE_TEST))
 		{
 			return CmdlineParseTest(*TestPAL::CommandLine);
 		}
+#endif
 	}
 
 	FPlatformMisc::SetCrashHandler(NULL);
