TestPAL: added a test of replacing an open file.

diff --git a/Engine/Source/Programs/TestPAL/Private/Main.cpp b/Engine/Source/Programs/TestPAL/Private/Main.cpp
--- a/Engine/Source/Programs/TestPAL/Private/Main.cpp
+++ b/Engine/Source/Programs/TestPAL/Private/Main.cpp
@@ -44,6 +44,12 @@ IMPLEMENT_APPLICATION(TestPAL, "TestPAL");
 #define ARG_EXEC_PROCESS_TEST				"exec-process"
 #define ARG_FORK_TEST						"fork"
 #define ARG_CMDLINE_PARSE_TEST				"cmdline"
+// Test that it is possible to replace a file open (for reading) in one Unreal process by another
+// Consists of two parts
+//    file-keep-open  - the instance that will open the file in one instance and will keep it so
+//    file-replace-open - the instance that will attempt to replace
+#define ARG_CMDLINE_FILE_KEEP_OPEN			"file-keep-open"
+#define ARG_CMDLINE_FILE_REPLACE_OPEN		"file-replace-open"
 
 namespace TestPAL
 {
@@ -1697,6 +1703,101 @@ int32 ExecProcessTest(const TCHAR *CommandLine)
 	return 0;
 }
 
+/**
+ * Opens the file for reading that is rewritable by another Unreal process and keeps it so.
+ *
+ */
+int32 FileKeepOpenTest(const TCHAR* CommandLine)
+{
+	FPlatformMisc::SetCrashHandler(NULL);
+	FPlatformMisc::SetGracefulTerminationHandler();
+
+	GEngineLoop.PreInit(CommandLine);
+
+	FString FileName;
+	if (!FParse::Value(CommandLine, TEXT("-file="), FileName))
+	{
+		UE_LOG(LogTestPAL, Error, TEXT("Use mandatory -file argument to point to the existing file"));
+		return 1;
+	}
+
+	UE_LOG(LogTestPAL, Display, TEXT("  File to keep open:'%s'"), *FileName);
+
+	{
+		TUniquePtr<IFileHandle> File(FPlatformFileManager::Get().GetPlatformFile().OpenRead(*FileName, true));
+		if (!File)
+		{
+			UE_LOG(LogTestPAL, Error, TEXT("Could not open file '%s'"), *FileName);
+			return 1;
+		}
+
+		UE_LOG(LogTestPAL, Display, TEXT("Opened the file, will keep it open until told to quit. Now use FileReplaceTest to replace it"), *FileName);
+
+		UE_LOG(LogTestPAL, Display, TEXT("IsEngineExitRequested = %d"), IsEngineExitRequested());
+
+		while(!IsEngineExitRequested())
+		{
+			FPlatformProcess::Sleep(1.0f);
+		}
+
+		UE_LOG(LogTestPAL, Display, TEXT("IsEngineExitRequested = %d"), IsEngineExitRequested());
+	}
+
+	FEngineLoop::AppPreExit();
+	FEngineLoop::AppExit();
+
+	return 0;
+}
+
+/**
+ * Attempts to replace the file that is open for reading by another TestPAL instance.
+ *
+ */
+int32 FileReplaceOpenTest(const TCHAR* CommandLine)
+{
+	FPlatformMisc::SetCrashHandler(NULL);
+	FPlatformMisc::SetGracefulTerminationHandler();
+
+	GEngineLoop.PreInit(CommandLine);
+
+	FString FileSourceName;
+	if (!FParse::Value(CommandLine, TEXT("-filesource="), FileSourceName))
+	{
+		UE_LOG(LogTestPAL, Error, TEXT("Use mandatory -filesource argument to point to the existing file"));
+		return 1;
+	}
+
+	FString FileTargetName;
+	if (!FParse::Value(CommandLine, TEXT("-filetarget="), FileTargetName))
+	{
+		UE_LOG(LogTestPAL, Error, TEXT("Use mandatory -filetarget argument to point to the existing file"));
+		return 1;
+	}
+
+	UE_LOG(LogTestPAL, Display, TEXT("  File that will be replaced:'%s'"), *FileTargetName);
+	UE_LOG(LogTestPAL, Display, TEXT("  File that it will be replaced with (i.e. source):'%s'"), *FileSourceName);
+
+	{
+		FString ToFileDir = FPaths::GetPath(FileTargetName);
+		FString TempDestFile = FPaths::CreateTempFilename(*ToFileDir);
+
+		// move the target file to a temp location first, then copy source over and delete the temp
+		if (!FPlatformFileManager::Get().GetPlatformFile().MoveFile(*TempDestFile, *FileTargetName) || !FPlatformFileManager::Get().GetPlatformFile().CopyFile(*FileTargetName, *FileSourceName) || 
+			!FPlatformFileManager::Get().GetPlatformFile().DeleteFile(*TempDestFile))
+		{
+			UE_LOG(LogTestPAL, Error, TEXT("Could not replace file '%s'!"), *FileTargetName);
+			return 1;
+		}
+
+		UE_LOG(LogTestPAL, Display, TEXT("Successfully replaced file '%s'!"), *FileTargetName);
+	}
+
+	FEngineLoop::AppPreExit();
+	FEngineLoop::AppExit();
+
+	return 0;
+}
+
 int32 CmdlineParseTest(const TCHAR *CommandLine)
 {
 	int32 NotOk = 0;	// zero when everything has passed so far
@@ -1844,6 +1945,14 @@ int32 MultiplexedMain(int32 ArgC, char* ArgV[])
 		{
 			return ForkTest(*TestPAL::CommandLine);
 		}
+		else if (!FCStringAnsi::Strcmp(ArgV[IdxArg], ARG_CMDLINE_FILE_KEEP_OPEN))
+		{
+			return FileKeepOpenTest(*TestPAL::CommandLine);
+		}
+		else if (!FCStringAnsi::Strcmp(ArgV[IdxArg], ARG_CMDLINE_FILE_REPLACE_OPEN))
+		{
+			return FileReplaceOpenTest(*TestPAL::CommandLine);
+		}
 #if PLATFORM_LINUX
 		else if (!FCStringAnsi::Strcmp(ArgV[IdxArg], ARG_CMDLINE_PARSE_TEST))
 		{
@@ -1879,6 +1988,11 @@ int32 MultiplexedMain(int32 ArgC, char* ArgV[])
 	UE_LOG(LogTestPAL, Warning, TEXT("  %s: test CreateGuid."), UTF8_TO_TCHAR(ARG_CREATEGUID_TEST));
 	UE_LOG(LogTestPAL, Warning, TEXT("  %s: test ThreadWalkStackAndDump and CaptureThreadBackTrace."), UTF8_TO_TCHAR(ARG_THREADSTACK_TEST));
 	UE_LOG(LogTestPAL, Warning, TEXT("  %s: test ExecProcess. Possible options: [-exe=executable] [-params=parameters]"), UTF8_TO_TCHAR(ARG_EXEC_PROCESS_TEST));
+	UE_LOG(LogTestPAL, Warning, TEXT("  %s: test replacing open files. Start the instance with %s and pass -file argument so it runs and keeps it open. This file is intended to be the target for '%s' test."), 
+		UTF8_TO_TCHAR(ARG_CMDLINE_FILE_KEEP_OPEN), UTF8_TO_TCHAR(ARG_CMDLINE_FILE_KEEP_OPEN), UTF8_TO_TCHAR(ARG_CMDLINE_FILE_REPLACE_OPEN));
+	UE_LOG(LogTestPAL, Warning, TEXT("  %s: test replacing open files. Run instance with %s giving it a -filesource and -filetarget arguments. Target file (expected to be kept open by '%s' test) will be replaced with the source."),
+		UTF8_TO_TCHAR(ARG_CMDLINE_FILE_REPLACE_OPEN), UTF8_TO_TCHAR(ARG_CMDLINE_FILE_REPLACE_OPEN), UTF8_TO_TCHAR(ARG_CMDLINE_FILE_KEEP_OPEN));
+
 
 	if (PLATFORM_LINUX)
 	{
