Fixed printf.

diff --git a/Engine/Source/Programs/TestPAL/Private/Main.cpp b/Engine/Source/Programs/TestPAL/Private/Main.cpp
--- a/Engine/Source/Programs/TestPAL/Private/Main.cpp
+++ b/Engine/Source/Programs/TestPAL/Private/Main.cpp
@@ -1731,7 +1731,7 @@ int32 FileKeepOpenTest(const TCHAR* CommandLine)
 			return 1;
 		}
 
-		UE_LOG(LogTestPAL, Display, TEXT("Opened the file, will keep it open until told to quit. Now use FileReplaceTest to replace it"), *FileName);
+		UE_LOG(LogTestPAL, Display, TEXT("Opened the file, will keep it open until told to quit. Now use FileReplaceTest to replace it"));
 
 		UE_LOG(LogTestPAL, Display, TEXT("IsEngineExitRequested = %d"), IsEngineExitRequested());
 
