Add new test for debugging a problem with deleting shared files in certain circumstances.

#Synced-CL 50828306

diff --git a/Engine/Source/Programs/TestPAL/Private/Main.cpp b/Engine/Source/Programs/TestPAL/Private/Main.cpp
--- a/Engine/Source/Programs/TestPAL/Private/Main.cpp
+++ b/Engine/Source/Programs/TestPAL/Private/Main.cpp
@@ -17,6 +17,7 @@
 #include "HAL/ThreadSafeCounter64.h"
 
 #include "Misc/Fork.h"
+#include "Async/MappedFileHandle.h"
 
 DEFINE_LOG_CATEGORY(LogTestPAL);
 
@@ -50,6 +51,7 @@ IMPLEMENT_APPLICATION(TestPAL, "TestPAL");
 //    file-replace-open - the instance that will attempt to replace
 #define ARG_CMDLINE_FILE_KEEP_OPEN			"file-keep-open"
 #define ARG_CMDLINE_FILE_REPLACE_OPEN		"file-replace-open"
+#define ARG_CMDLINE_FILE_DELETE_MAPPED		"file-delete-mapped"	// self-contained test, opens a mapped file and tries to delete it
 
 namespace TestPAL
 {
@@ -1798,6 +1800,83 @@ int32 FileReplaceOpenTest(const TCHAR* CommandLine)
 	return 0;
 }
 
+
+/**
+ * Creates a temporary file (right where the binary is located), maps it open, then attempts to delete it while keeping the handle open.
+ *
+ */
+int32 FileDeleteMappedTest(const TCHAR* CommandLine)
+{
+	FPlatformMisc::SetCrashHandler(NULL);
+	FPlatformMisc::SetGracefulTerminationHandler();
+
+	GEngineLoop.PreInit(CommandLine);
+
+	FString MmappedFile = FPaths::CreateTempFilename(TEXT("."));
+
+	UE_LOG(LogTestPAL, Display, TEXT("  File that will be kept memory-mapped:'%s'"), *MmappedFile);
+	{
+		FRandomStream Rand(static_cast<int32>(FPlatformTime::Cycles()));
+
+		auto GetRandomString = [Rand]() -> FString
+			{
+				int32 Size = static_cast<int32>(128 + 1024.0 * Rand.GetFraction()) * 1024;
+				FString Temp;
+
+				Temp.Reserve(Size);
+				for (int32 Idx = 0; Idx < Size; ++Idx)
+				{
+					Temp += static_cast<TCHAR>(static_cast<uint32>(TEXT('0')) + static_cast<uint32>(36.0 * Rand.GetFraction()) * sizeof(TCHAR));
+				}
+
+				return MoveTemp(Temp);
+			};
+
+		auto SaveRandomFile = [Rand, GetRandomString](const FString& Filename) -> bool
+			{
+				FString Temp = GetRandomString();
+				return FFileHelper::SaveStringToFile(Temp, *Filename, FFileHelper::EEncodingOptions::ForceUTF8);
+			};
+
+		if (!SaveRandomFile(MmappedFile))
+		{
+			UE_LOG(LogTestPAL, Error, TEXT("Could not create file '%s'!"), *MmappedFile);
+			return 1;
+		}
+	}
+
+	{
+		TUniquePtr<IMappedFileHandle> Handle;
+		{
+			FOpenMappedResult Result = FPlatformFileManager::Get().GetPlatformFile().OpenMappedEx(*MmappedFile, IPlatformFile::EOpenReadFlags::AllowDelete);
+			if (Result.HasError())
+			{
+				UE_LOG(LogTestPAL, Error, TEXT("Could not memory map file '%s'!"), *MmappedFile);
+				return 1;
+			}
+
+			Handle = Result.StealValue();
+		}
+
+		// delete, or infinitely loop
+		IPlatformFile& PlatformFile = FPlatformFileManager::Get().GetPlatformFile();
+
+		UE_LOG(LogTestPAL, Display, TEXT("Attempting to delete file '%s', will loop forever until it suceeds to help analyze problems."), *MmappedFile);
+		int64 NumTry = 1;
+		// intersperse exists and delete to have a better chance for CPU profiling to find a boundary for the DeleteFile
+		while (PlatformFile.FileExists(*MmappedFile) && !PlatformFile.DeleteFile(*MmappedFile))
+		{
+			++NumTry;
+		}
+		UE_LOG(LogTestPAL, Display, TEXT("Successfully deleted file '%s' on a %lld try!"), *MmappedFile, NumTry);
+	}
+
+	FEngineLoop::AppPreExit();
+	FEngineLoop::AppExit();
+
+	return 0;
+}
+
 int32 CmdlineParseTest(const TCHAR *CommandLine)
 {
 	int32 NotOk = 0;	// zero when everything has passed so far
@@ -1953,6 +2032,10 @@ int32 MultiplexedMain(int32 ArgC, char* ArgV[])
 		{
 			return FileReplaceOpenTest(*TestPAL::CommandLine);
 		}
+		else if (!FCStringAnsi::Strcmp(ArgV[IdxArg], ARG_CMDLINE_FILE_DELETE_MAPPED))
+		{
+			return FileDeleteMappedTest(*TestPAL::CommandLine);
+		}
 #if PLATFORM_LINUX
 		else if (!FCStringAnsi::Strcmp(ArgV[IdxArg], ARG_CMDLINE_PARSE_TEST))
 		{
@@ -1992,7 +2075,7 @@ int32 MultiplexedMain(int32 ArgC, char* ArgV[])
 		UTF8_TO_TCHAR(ARG_CMDLINE_FILE_KEEP_OPEN), UTF8_TO_TCHAR(ARG_CMDLINE_FILE_KEEP_OPEN), UTF8_TO_TCHAR(ARG_CMDLINE_FILE_REPLACE_OPEN));
 	UE_LOG(LogTestPAL, Warning, TEXT("  %s: test replacing open files. Run instance with %s giving it a -filesource and -filetarget arguments. Target file (expected to be kept open by '%s' test) will be replaced with the source."),
 		UTF8_TO_TCHAR(ARG_CMDLINE_FILE_REPLACE_OPEN), UTF8_TO_TCHAR(ARG_CMDLINE_FILE_REPLACE_OPEN), UTF8_TO_TCHAR(ARG_CMDLINE_FILE_KEEP_OPEN));
-
+	UE_LOG(LogTestPAL, Warning, TEXT("  %s: test replacing memory-mapped files within the same process. Creates a temp file, maps it into the memory, then attempts to replace with a new one."), UTF8_TO_TCHAR(ARG_CMDLINE_FILE_DELETE_MAPPED));
 
 	if (PLATFORM_LINUX)
 	{
