Implement KillTree in FUnixPlatformProcess::TerminateProc

#rnx
#Synced-CL 52574197

diff --git a/Engine/Source/Programs/TestPAL/Private/Main.cpp b/Engine/Source/Programs/TestPAL/Private/Main.cpp
--- a/Engine/Source/Programs/TestPAL/Private/Main.cpp
+++ b/Engine/Source/Programs/TestPAL/Private/Main.cpp
@@ -25,6 +25,10 @@ IMPLEMENT_APPLICATION(TestPAL, "TestPAL");
 
 #define ARG_PROC_TEST						"proc"
 #define ARG_PROC_TEST_CHILD					"proc-child"
+#define ARG_PROC_KILLTREE_TEST				"proc-killtree"
+#define ARG_PROC_KILLTREE_TEST_CHILD		"proc-killtree-child"
+#define ARG_PROC_KILLTREE_TEST_GRANDCHILD	"proc-killtree-grandchild"
+#define ARG_PROC_OPENPROCESS_TEST			"proc-openprocess"
 #define ARG_CASE_SENSITIVITY_TEST			"case"
 #define ARG_MESSAGEBOX_TEST					"messagebox"
 #define ARG_DIRECTORY_WATCHER_TEST			"dirwatcher"
@@ -117,6 +121,363 @@ int32 ProcRunAsParent(const TCHAR* CommandLine)
 	return 0;
 }
 
+/**
+ * TerminateProc KillTree test (grandchild instance).
+ * Spins until signalled. Installed as a leaf process in the tree under proc-killtree-child.
+ */
+int32 ProcKillTreeRunAsGrandchild(const TCHAR* CommandLine)
+{
+	FPlatformMisc::SetCrashHandler(NULL);
+	FPlatformMisc::SetGracefulTerminationHandler();
+
+	GEngineLoop.PreInit(CommandLine);
+	UE_LOG(LogTestPAL, Display, TEXT("Running proc-killtree test as grandchild (pid %d), spinning until killed."),
+		FPlatformProcess::GetCurrentProcessId());
+
+	while (!IsEngineExitRequested())
+	{
+		FPlatformProcess::Sleep(0.25f);
+	}
+
+	FEngineLoop::AppPreExit();
+	FEngineLoop::AppExit();
+	return 0;
+}
+
+/**
+ * TerminateProc KillTree test (child instance).
+ * Spawns a grandchild, writes the grandchild's PID to the file specified by -pidfile=, then spins until signalled.
+ */
+int32 ProcKillTreeRunAsChild(const TCHAR* CommandLine)
+{
+	FPlatformMisc::SetCrashHandler(NULL);
+	FPlatformMisc::SetGracefulTerminationHandler();
+
+	GEngineLoop.PreInit(CommandLine);
+
+	FString PidFilePath;
+	if (!FParse::Value(CommandLine, TEXT("pidfile="), PidFilePath))
+	{
+		UE_LOG(LogTestPAL, Error, TEXT("proc-killtree-child requires a -pidfile=<path> argument."));
+		FEngineLoop::AppPreExit();
+		FEngineLoop::AppExit();
+		return 1;
+	}
+
+	UE_LOG(LogTestPAL, Display, TEXT("Running proc-killtree test as child (pid %d), spawning grandchild."),
+		FPlatformProcess::GetCurrentProcessId());
+
+	FString WorkerName = FPlatformProcess::ExecutablePath();
+	uint32 GrandchildPid = 0;
+	FProcHandle GrandchildHandle = FPlatformProcess::CreateProc(*WorkerName, TEXT(ARG_PROC_KILLTREE_TEST_GRANDCHILD),
+		false, false, false, &GrandchildPid, -1, nullptr, nullptr);
+
+	if (!GrandchildHandle.IsValid())
+	{
+		UE_LOG(LogTestPAL, Error, TEXT("proc-killtree-child: failed to spawn grandchild."));
+		FEngineLoop::AppPreExit();
+		FEngineLoop::AppExit();
+		return 1;
+	}
+
+	UE_LOG(LogTestPAL, Display, TEXT("Child (pid %d) spawned grandchild (pid %u), writing PID to '%s'."),
+		FPlatformProcess::GetCurrentProcessId(), GrandchildPid, *PidFilePath);
+
+	// Write the grandchild PID to the agreed-upon file so the parent can verify it was killed
+	if (!FFileHelper::SaveStringToFile(FString::Printf(TEXT("%u"), GrandchildPid), *PidFilePath))
+	{
+		UE_LOG(LogTestPAL, Error, TEXT("proc-killtree-child: failed to write PID file '%s', terminating grandchild."), *PidFilePath);
+		FPlatformProcess::TerminateProc(GrandchildHandle, true);
+		FPlatformProcess::WaitForProc(GrandchildHandle);
+		FPlatformProcess::CloseProc(GrandchildHandle);
+		FEngineLoop::AppPreExit();
+		FEngineLoop::AppExit();
+		return 1;
+	}
+
+	// Spin until the kill tree signal arrives
+	while (!IsEngineExitRequested())
+	{
+		FPlatformProcess::Sleep(0.25f);
+	}
+
+	FPlatformProcess::WaitForProc(GrandchildHandle);
+	FPlatformProcess::CloseProc(GrandchildHandle);
+
+	FEngineLoop::AppPreExit();
+	FEngineLoop::AppExit();
+	return 0;
+}
+
+/**
+ * TerminateProc KillTree test (parent / driver instance).
+ * Spawns a child which itself spawns a grandchild, then verifies that calling
+ * TerminateProc(..., KillTree=true) kills both.
+ */
+int32 ProcKillTreeRunAsParent(const TCHAR* CommandLine)
+{
+	FPlatformMisc::SetCrashHandler(NULL);
+	FPlatformMisc::SetGracefulTerminationHandler();
+
+	GEngineLoop.PreInit(CommandLine);
+	UE_LOG(LogTestPAL, Display, TEXT("Running proc-killtree test as parent (pid %d)."),
+		FPlatformProcess::GetCurrentProcessId());
+
+	// Unique PID file the child will write its grandchild's PID into
+	FString PidFilePath = FPaths::Combine(
+		FPlatformProcess::UserTempDir(),
+		FString::Printf(TEXT("testpal-killtree-%u.pid"), FPlatformProcess::GetCurrentProcessId()));
+
+	// Spawn the child, forwarding the PID file path
+	FString WorkerName = FPlatformProcess::ExecutablePath();
+	FString ChildArgs = FString::Printf(TEXT(ARG_PROC_KILLTREE_TEST_CHILD) TEXT(" -pidfile=\"%s\""), *PidFilePath);
+	uint32 ChildPid = 0;
+	FProcHandle ChildHandle = FPlatformProcess::CreateProc(*WorkerName, *ChildArgs,
+		false, false, false, &ChildPid, -1, nullptr, nullptr);
+
+	if (!ChildHandle.IsValid())
+	{
+		UE_LOG(LogTestPAL, Error, TEXT("Failed to spawn child process."));
+		FEngineLoop::AppPreExit();
+		FEngineLoop::AppExit();
+		return 1;
+	}
+
+	UE_LOG(LogTestPAL, Display, TEXT("Spawned child (pid %u). Waiting for grandchild PID file '%s'."),
+		ChildPid, *PidFilePath);
+
+	// Poll for the PID file written by the child once its grandchild is running
+	const double PidFileTimeoutSeconds = 15.0;
+	double StartTime = FPlatformTime::Seconds();
+	while (!FPaths::FileExists(PidFilePath))
+	{
+		if (FPlatformTime::Seconds() - StartTime > PidFileTimeoutSeconds)
+		{
+			UE_LOG(LogTestPAL, Error, TEXT("Timed out waiting for grandchild PID file."));
+			FPlatformProcess::TerminateProc(ChildHandle, true);
+			FPlatformProcess::WaitForProc(ChildHandle);
+			FPlatformProcess::CloseProc(ChildHandle);
+			IFileManager::Get().Delete(*PidFilePath);
+			FEngineLoop::AppPreExit();
+			FEngineLoop::AppExit();
+			return 1;
+		}
+		FPlatformProcess::Sleep(0.05f);
+	}
+
+	// Brief pause to let the child finish writing and close the file handle;
+	// FileExists can return true before the write is fully flushed.
+	FPlatformProcess::Sleep(0.25f);
+
+	FString GrandchildPidStr;
+	if (!FFileHelper::LoadFileToString(GrandchildPidStr, *PidFilePath))
+	{
+		UE_LOG(LogTestPAL, Error, TEXT("Failed to read grandchild PID file '%s'."), *PidFilePath);
+		FPlatformProcess::TerminateProc(ChildHandle, true);
+		FPlatformProcess::WaitForProc(ChildHandle);
+		FPlatformProcess::CloseProc(ChildHandle);
+		IFileManager::Get().Delete(*PidFilePath);
+		FEngineLoop::AppPreExit();
+		FEngineLoop::AppExit();
+		return 1;
+	}
+
+	uint32 GrandchildPid = static_cast<uint32>(FCString::Atoi(*GrandchildPidStr));
+	if (GrandchildPid == 0)
+	{
+		UE_LOG(LogTestPAL, Error, TEXT("Grandchild PID file contained invalid data: '%s'."), *GrandchildPidStr);
+		FPlatformProcess::TerminateProc(ChildHandle, true);
+		FPlatformProcess::WaitForProc(ChildHandle);
+		FPlatformProcess::CloseProc(ChildHandle);
+		IFileManager::Get().Delete(*PidFilePath);
+		FEngineLoop::AppPreExit();
+		FEngineLoop::AppExit();
+		return 1;
+	}
+	UE_LOG(LogTestPAL, Display, TEXT("Grandchild PID is %u. Verifying both processes are running before the kill."),
+		GrandchildPid);
+
+	// Sanity-check: both must be alive before we test the kill
+	if (!FPlatformProcess::IsProcRunning(ChildHandle))
+	{
+		UE_LOG(LogTestPAL, Error, TEXT("Child (pid %u) is not running before TerminateProc - test invalid."), ChildPid);
+		FPlatformProcess::CloseProc(ChildHandle);
+		IFileManager::Get().Delete(*PidFilePath);
+		FEngineLoop::AppPreExit();
+		FEngineLoop::AppExit();
+		return 1;
+	}
+
+	if (!FPlatformProcess::IsApplicationRunning(GrandchildPid))
+	{
+		UE_LOG(LogTestPAL, Error, TEXT("Grandchild (pid %u) is not running before TerminateProc - test invalid."), GrandchildPid);
+		FPlatformProcess::TerminateProc(ChildHandle, true);
+		FPlatformProcess::WaitForProc(ChildHandle);
+		FPlatformProcess::CloseProc(ChildHandle);
+		IFileManager::Get().Delete(*PidFilePath);
+		FEngineLoop::AppPreExit();
+		FEngineLoop::AppExit();
+		return 1;
+	}
+
+	UE_LOG(LogTestPAL, Display, TEXT("Child (%u) and grandchild (%u) both confirmed running. Calling TerminateProc with KillTree=true."),
+		ChildPid, GrandchildPid);
+
+	FPlatformProcess::TerminateProc(ChildHandle, true /* KillTree */);
+
+	// Poll for the child to die (it exits gracefully on SIGTERM)
+	const double KillTimeoutSeconds = 10.0;
+	StartTime = FPlatformTime::Seconds();
+	while (FPlatformProcess::IsProcRunning(ChildHandle) && FPlatformTime::Seconds() - StartTime < KillTimeoutSeconds)
+	{
+		FPlatformProcess::Sleep(0.1f);
+	}
+	bool bChildDead = !FPlatformProcess::IsProcRunning(ChildHandle);
+
+	// Poll for the grandchild to die and be reaped (adopted by init after child exits)
+	StartTime = FPlatformTime::Seconds();
+	bool bGrandchildDead = false;
+	while (FPlatformTime::Seconds() - StartTime < KillTimeoutSeconds)
+	{
+		if (!FPlatformProcess::IsApplicationRunning(GrandchildPid))
+		{
+			bGrandchildDead = true;
+			break;
+		}
+		FPlatformProcess::Sleep(0.1f);
+	}
+
+	if (bChildDead)
+	{
+		FPlatformProcess::WaitForProc(ChildHandle);
+	}
+
+	// Best-effort cleanup of any surviving processes so CI doesn't leak them.
+	if (!bChildDead)
+	{
+		UE_LOG(LogTestPAL, Error, TEXT("FAILED: child (pid %u) is still running after TerminateProc KillTree=true."), ChildPid);
+		FPlatformProcess::TerminateProc(ChildHandle, true);
+		FPlatformProcess::WaitForProc(ChildHandle);
+	}
+	if (!bGrandchildDead)
+	{
+		UE_LOG(LogTestPAL, Error, TEXT("FAILED: grandchild (pid %u) is still running after TerminateProc KillTree=true."), GrandchildPid);
+		FProcHandle GrandchildHandle = FPlatformProcess::OpenProcess(GrandchildPid);
+		if (GrandchildHandle.IsValid())
+		{
+			FPlatformProcess::TerminateProc(GrandchildHandle, false);
+			FPlatformProcess::WaitForProc(GrandchildHandle);
+			FPlatformProcess::CloseProc(GrandchildHandle);
+		}
+	}
+
+	FPlatformProcess::CloseProc(ChildHandle);
+	IFileManager::Get().Delete(*PidFilePath);
+
+	if (bChildDead && bGrandchildDead)
+	{
+		UE_LOG(LogTestPAL, Display, TEXT("SUCCESS: child (pid %u) and grandchild (pid %u) were both terminated by TerminateProc KillTree=true."),
+			ChildPid, GrandchildPid);
+	}
+
+	FEngineLoop::AppPreExit();
+	FEngineLoop::AppExit();
+	return (bChildDead && bGrandchildDead) ? 0 : 1;
+}
+
+/**
+ * OpenProcess lifecycle test.
+ * Spawns a child, obtains an OpenProcess handle, then exercises
+ * IsProcRunning, TerminateProc (KillTree), WaitForProc, and CloseProc
+ * through the OpenProcess code path.
+ */
+int32 ProcOpenProcessTest(const TCHAR* CommandLine)
+{
+	FPlatformMisc::SetCrashHandler(NULL);
+	FPlatformMisc::SetGracefulTerminationHandler();
+
+	GEngineLoop.PreInit(CommandLine);
+
+	UE_LOGF(LogTestPAL, Display, "Running proc-openprocess test (pid %d).",
+		FPlatformProcess::GetCurrentProcessId());
+
+	// Spawn a child that spins until killed (reuse the grandchild spinner role)
+	FString WorkerName = FPlatformProcess::ExecutablePath();
+	uint32 ChildPid = 0;
+	FProcHandle SpawnHandle = FPlatformProcess::CreateProc(*WorkerName, TEXT(ARG_PROC_KILLTREE_TEST_GRANDCHILD),
+		false, false, false, &ChildPid, -1, nullptr, nullptr);
+
+	if (!SpawnHandle.IsValid())
+	{
+		UE_LOGF(LogTestPAL, Error, "FAILED: could not spawn child process.");
+		FEngineLoop::AppPreExit();
+		FEngineLoop::AppExit();
+		return 1;
+	}
+
+	UE_LOGF(LogTestPAL, Display, "Spawned child (pid %u). Waiting briefly for it to start.", ChildPid);
+	FPlatformProcess::Sleep(0.5f);
+
+	// Open a second handle to the same process via OpenProcess (the code path under test)
+	FProcHandle OpenHandle = FPlatformProcess::OpenProcess(ChildPid);
+	if (!OpenHandle.IsValid())
+	{
+		UE_LOGF(LogTestPAL, Error, "FAILED: OpenProcess(%u) returned invalid handle.", ChildPid);
+		FPlatformProcess::TerminateProc(SpawnHandle, true);
+		FPlatformProcess::WaitForProc(SpawnHandle);
+		FPlatformProcess::CloseProc(SpawnHandle);
+		FEngineLoop::AppPreExit();
+		FEngineLoop::AppExit();
+		return 1;
+	}
+
+	// Verify IsProcRunning returns true through the OpenProcess handle
+	if (!FPlatformProcess::IsProcRunning(OpenHandle))
+	{
+		UE_LOGF(LogTestPAL, Error, "FAILED: IsProcRunning returned false for a live process via OpenProcess handle.");
+		FPlatformProcess::TerminateProc(SpawnHandle, true);
+		FPlatformProcess::WaitForProc(SpawnHandle);
+		FPlatformProcess::CloseProc(SpawnHandle);
+		FPlatformProcess::CloseProc(OpenHandle);
+		FEngineLoop::AppPreExit();
+		FEngineLoop::AppExit();
+		return 1;
+	}
+
+	UE_LOGF(LogTestPAL, Display, "IsProcRunning(OpenHandle) confirmed child alive. Calling TerminateProc(OpenHandle, KillTree=true).");
+
+	// Kill the child through the OpenProcess handle with KillTree
+	FPlatformProcess::TerminateProc(OpenHandle, true);
+
+	// Wait for it through the OpenProcess handle (detects zombie state and returns)
+	FPlatformProcess::WaitForProc(OpenHandle);
+
+	// Reap the child via the SpawnHandle so it transitions from zombie to fully gone.
+	// IsProcRunning uses kill(pid, 0) which returns 0 for zombies, so we must reap first.
+	FPlatformProcess::WaitForProc(SpawnHandle);
+
+	// Verify it's dead through the OpenProcess handle
+	if (FPlatformProcess::IsProcRunning(OpenHandle))
+	{
+		UE_LOGF(LogTestPAL, Error, "FAILED: child (pid %u) still running after TerminateProc+WaitForProc via OpenProcess handle.", ChildPid);
+		FPlatformProcess::CloseProc(OpenHandle);
+		// SpawnHandle was already reaped via WaitForProc above; just close it.
+		FPlatformProcess::CloseProc(SpawnHandle);
+		FEngineLoop::AppPreExit();
+		FEngineLoop::AppExit();
+		return 1;
+	}
+
+	UE_LOGF(LogTestPAL, Display, "SUCCESS: child (pid %u) terminated and confirmed dead via OpenProcess handle.", ChildPid);
+
+	FPlatformProcess::CloseProc(OpenHandle);
+	FPlatformProcess::CloseProc(SpawnHandle);
+
+	FEngineLoop::AppPreExit();
+	FEngineLoop::AppExit();
+	return 0;
+}
+
 /**
  * Tests a single file.
  */
@@ -1946,6 +2307,22 @@ int32 MultiplexedMain(int32 ArgC, char* ArgV[])
 		{
 			return ProcRunAsParent(*TestPAL::CommandLine);
 		}
+		else if (!FCStringAnsi::Strcmp(ArgV[IdxArg], ARG_PROC_KILLTREE_TEST_GRANDCHILD))
+		{
+			return ProcKillTreeRunAsGrandchild(*TestPAL::CommandLine);
+		}
+		else if (!FCStringAnsi::Strcmp(ArgV[IdxArg], ARG_PROC_KILLTREE_TEST_CHILD))
+		{
+			return ProcKillTreeRunAsChild(*TestPAL::CommandLine);
+		}
+		else if (!FCStringAnsi::Strcmp(ArgV[IdxArg], ARG_PROC_KILLTREE_TEST))
+		{
+			return ProcKillTreeRunAsParent(*TestPAL::CommandLine);
+		}
+		else if (!FCStringAnsi::Strcmp(ArgV[IdxArg], ARG_PROC_OPENPROCESS_TEST))
+		{
+			return ProcOpenProcessTest(*TestPAL::CommandLine);
+		}
 		else if (!FCStringAnsi::Strcmp(ArgV[IdxArg], ARG_CASE_SENSITIVITY_TEST))
 		{
 			return CaseTest(*TestPAL::CommandLine);
@@ -2053,6 +2430,8 @@ int32 MultiplexedMain(int32 ArgC, char* ArgV[])
 	UE_LOG(LogTestPAL, Warning, TEXT(""));
 	UE_LOG(LogTestPAL, Warning, TEXT("Available test cases:"));
 	UE_LOG(LogTestPAL, Warning, TEXT("  %s: test process handling API"), ANSI_TO_TCHAR( ARG_PROC_TEST ));
+	UE_LOG(LogTestPAL, Warning, TEXT("  %s: test TerminateProc with KillTree=true (spawns a child which spawns a grandchild, then verifies both are killed)"), ANSI_TO_TCHAR( ARG_PROC_KILLTREE_TEST ));
+	UE_LOG(LogTestPAL, Warning, TEXT("  %s: test OpenProcess handle lifecycle (IsProcRunning, TerminateProc, WaitForProc via OpenProcess)"), ANSI_TO_TCHAR( ARG_PROC_OPENPROCESS_TEST ));
 	UE_LOG(LogTestPAL, Warning, TEXT("  %s: test case-insensitive file operations"), ANSI_TO_TCHAR( ARG_CASE_SENSITIVITY_TEST ));
 	UE_LOG(LogTestPAL, Warning, TEXT("  %s: test message box bug (too long strings)"), ANSI_TO_TCHAR( ARG_MESSAGEBOX_TEST ));
 	UE_LOG(LogTestPAL, Warning, TEXT("  %s: test directory watcher"), ANSI_TO_TCHAR( ARG_DIRECTORY_WATCHER_TEST ));
diff --git a/Engine/Source/Runtime/Core/Private/Unix/UnixPlatformProcess.cpp b/Engine/Source/Runtime/Core/Private/Unix/UnixPlatformProcess.cpp
--- a/Engine/Source/Runtime/Core/Private/Unix/UnixPlatformProcess.cpp
+++ b/Engine/Source/Runtime/Core/Private/Unix/UnixPlatformProcess.cpp
@@ -16,6 +16,8 @@
 #include "Misc/Paths.h"
 #include "Unix/UnixPlatformRunnableThread.h"
 #include "Misc/EngineVersion.h"
+#include <dirent.h>
+#include <signal.h>
 #include <spawn.h>
 #include <sys/wait.h>
 #include <sys/resource.h>
@@ -1122,7 +1124,7 @@ FProcHandle FUnixPlatformProcess::CreateProcInternal(const TCHAR* URL, const TCH
 
 	if (PosixSpawnErrNo != 0)
 	{
-		UE_LOGF(LogHAL, Error, "FUnixPlatformProcess::CreateProc: posix_spawnp() failed (%d, %ls)", PosixSpawnErrNo, UTF8_TO_TCHAR(strerror(PosixSpawnErrNo)));
+		UE_LOGF(LogHAL, Error, "FUnixPlatformProcess::CreateProc: posix_spawnp() failed (%d, %s) Path: %ls, Parms: %ls", PosixSpawnErrNo, strerror(PosixSpawnErrNo), *ProcessPath, Parms);
 		return FProcHandle();
 	}
 
@@ -1212,14 +1214,116 @@ FProcHandle FUnixPlatformProcess::CreateProcInternal(const TCHAR* URL, const TCH
 	return FProcHandle(new FProcState(ChildPid, bLaunchDetached));
 }
 
-/*
- * Return a limited use FProcHandle from a PID. Currently can only use w/ IsProcRunning().
- *
- * WARNING (from Arciel): PIDs can and will be reused. We have had that issue
- * before: the editor was tracking ShaderCompileWorker by their PIDs, and a
- * long-running process (something from PS4 SDK) got a reused SCW PID,
- * resulting in compilation never ending.
+/**
+ * Read the start time of a process from /proc/<pid>/stat field 22 (starttime).
+ * Returns clock ticks since boot, or 0 on failure. PID + starttime uniquely
+ * identifies a process instance even across PID reuse.
+ */
+static uint64 ReadProcessStartTime(pid_t Pid)
+{
+	// "/proc/" (6) + max PID digits (7) + "/stat" (5) + '\0' (1) = 19; 32 is generous
+	char StatPath[32];
+	FCStringAnsi::Sprintf(StatPath, "/proc/%d/stat", static_cast<int>(Pid));
+
+	FILE* StatFile = fopen(StatPath, "r");
+	if (StatFile == nullptr)
+	{
+		return 0;
+	}
+
+	// /proc/<pid>/stat is a single line; 1024 is generous for any process name
+	char Line[1024];
+	bool bReadOk = (fgets(Line, sizeof(Line), StatFile) != nullptr);
+	fclose(StatFile);
+
+	if (!bReadOk)
+	{
+		return 0;
+	}
+
+	// Field 2 is (comm) which can contain spaces and parentheses.
+	// Find the last ')' to safely skip past it, then count fields from field 3.
+	const char* CloseParen = strrchr(Line, ')');
+	if (CloseParen == nullptr)
+	{
+		return 0;
+	}
+
+	// Field 22 (starttime) is 20 fields after field 2.
+	// After ')' there is a space, then field 3 starts.
+	const char* Pos = CloseParen + 1;
+	for (int Field = 3; Field < 22; ++Field)
+	{
+		while (*Pos == ' ')
+		{
+			++Pos;
+		}
+		if (*Pos == '\0')
+		{
+			return 0;
+		}
+		while (*Pos != ' ' && *Pos != '\0')
+		{
+			++Pos;
+		}
+	}
+	while (*Pos == ' ')
+	{
+		++Pos;
+	}
+
+	if (*Pos == '\0')
+	{
+		return 0;
+	}
+
+	return static_cast<uint64>(strtoull(Pos, nullptr, 10));
+}
+
+/**
+ * Validate that an OpenProcess handle still refers to the same process.
+ * For CreateProc handles (ProcInfo != nullptr), always succeeds.
+ * For OpenProcess handles, compares the stored starttime against the current
+ * value in /proc to detect PID reuse. Sets OutPid on success.
  */
+static bool ValidateOpenProcessHandle(const FProcHandle& Handle, pid_t& OutPid)
+{
+	if (Handle.GetProcessInfo() != nullptr)
+	{
+		OutPid = Handle.GetProcessInfo()->GetProcessId();
+		return true;
+	}
+
+	if (Handle.Get() == -1)
+	{
+		return false;
+	}
+
+	OutPid = static_cast<pid_t>(Handle.Get());
+
+	// If we didn't capture a starttime (e.g. /proc was unavailable), skip validation
+	if (Handle.OpenedStartTime == 0)
+	{
+		return true;
+	}
+
+	uint64 CurrentStartTime = ReadProcessStartTime(OutPid);
+	if (CurrentStartTime == 0)
+	{
+		// Process no longer exists (can't read /proc/<pid>/stat)
+		return false;
+	}
+
+	if (CurrentStartTime != Handle.OpenedStartTime)
+	{
+		UE_LOGF(LogHAL, Warning, "OpenProcess handle for pid %d detected PID reuse (starttime %" UINT64_FMT " vs %" UINT64_FMT "). Ignoring stale handle.",
+			static_cast<int32>(OutPid), Handle.OpenedStartTime, CurrentStartTime);
+		return false;
+	}
+
+	return true;
+}
+
 FProcHandle FUnixPlatformProcess::OpenProcess(uint32 ProcessID)
 {
 	pid_t Pid = static_cast< pid_t >(ProcessID);
@@ -1231,7 +1335,12 @@ FProcHandle FUnixPlatformProcess::OpenProcess(uint32 ProcessID)
 	// errno == EPERM: don't have permissions to send signal
 	// errno == ESRCH: proc doesn't exist
 	bool bIsRunning = (KillResult == 0);
-	return FProcHandle(bIsRunning ? Pid : -1);
+	FProcHandle Handle(bIsRunning ? Pid : -1);
+	if (bIsRunning)
+	{
+		Handle.OpenedStartTime = ReadProcessStartTime(Pid);
+	}
+	return Handle;
 }
 
 /** Initialization constructor. */
@@ -1386,25 +1495,21 @@ void FProcState::Wait()
 
 bool FUnixPlatformProcess::IsProcRunning( FProcHandle & ProcessHandle )
 {
-	bool bIsRunning = false;
 	FProcState * ProcInfo = ProcessHandle.GetProcessInfo();
-
 	if (ProcInfo)
 	{
-		bIsRunning = ProcInfo->IsRunning();
+		return ProcInfo->IsRunning();
 	}
-	else if (ProcessHandle.Get() != -1)
-	{
-		// Process opened with OpenProcess() call (we only have pid)
-		int KillResult = kill(ProcessHandle.Get(), 0);	// no actual signal is sent
-		check(KillResult != -1 || errno != EINVAL);
 
-		// errno == EPERM: don't have permissions to send signal
-		// errno == ESRCH: proc doesn't exist
-		bIsRunning = (KillResult == 0);
+	pid_t Pid;
+	if (!ValidateOpenProcessHandle(ProcessHandle, Pid))
+	{
+		return false;
 	}
 
-	return bIsRunning;
+	int KillResult = kill(Pid, 0);
+	check(KillResult != -1 || errno != EINVAL);
+	return (KillResult == 0);
 }
 
 void FUnixPlatformProcess::WaitForProc( FProcHandle & ProcessHandle )
@@ -1414,9 +1519,51 @@ void FUnixPlatformProcess::WaitForProc( FProcHandle & ProcessHandle )
 	{
 		ProcInfo->Wait();
 	}
-	else if (ProcessHandle.Get() != -1)
+	else
 	{
-		STUBBED("FUnixPlatformProcess::WaitForProc() : Waiting on OpenProcess() handle not implemented yet");
+		pid_t Pid;
+		if (!ValidateOpenProcessHandle(ProcessHandle, Pid))
+		{
+			return;
+		}
+
+		// Poll /proc/<pid>/status until the process disappears or becomes a zombie.
+		// "/proc/" (6) + max PID digits (7, up to 4194304) + "/status" (7) + '\0' (1) = 21; 48 is generous
+		char StatusPath[48];
+		FCStringAnsi::Sprintf(StatusPath, "/proc/%d/status", static_cast<int>(Pid));
+		while (true)
+		{
+			FILE* StatusFile = fopen(StatusPath, "r");
+			if (StatusFile == nullptr)
+			{
+				break; // process no longer exists
+			}
+
+			bool bIsZombie = false;
+			// Longest line in /proc/pid/status is typically under 80 chars; 256 is generous
+			char Line[256];
+			while (fgets(Line, sizeof(Line), StatusFile) != nullptr)
+			{
+				if (strncmp(Line, "State:", 6) == 0)
+				{
+					// State line format: "State:\t<char> (<description>)"
+					const char* StateChar = Line + 6;
+					while (*StateChar != '\0' && (*StateChar == '\t' || *StateChar == ' '))
+					{
+						++StateChar;
+					}
+					bIsZombie = (*StateChar == 'Z');
+					break;
+				}
+			}
+			fclose(StatusFile);
+
+			if (bIsZombie)
+			{
+				break;
+			}
+			FPlatformProcess::Sleep(0.05f);
+		}
 	}
 }
 
@@ -1431,21 +1578,142 @@ void FUnixPlatformProcess::CloseProc(FProcHandle & ProcessHandle)
 
 void FUnixPlatformProcess::TerminateProc( FProcHandle & ProcessHandle, bool KillTree )
 {
-	if (KillTree)
+	FProcState* ProcInfo = ProcessHandle.GetProcessInfo();
+
+	pid_t RootPid = -1;
+	if (ProcInfo != nullptr)
 	{
-		// TODO: enumerate the children
-		STUBBED("FUnixPlatformProcess::TerminateProc() : Killing a subtree is not implemented yet");
+		RootPid = ProcInfo->GetProcessId();
+	}
+	else if (!ValidateOpenProcessHandle(ProcessHandle, RootPid))
+	{
+		return;
 	}
 
-	FProcState * ProcInfo = ProcessHandle.GetProcessInfo();
-	if (ProcInfo)
+	if (KillTree && RootPid != -1)
 	{
-		int KillResult = kill(ProcInfo->GetProcessId(), SIGTERM);	// graceful
-		check(KillResult != -1 || errno != EINVAL);
+		// Iteratively collect all descendant PIDs via BFS, seeded by one /proc scan per pass.
+		// Two passes are used: the first finds the initial tree; the second catches any processes
+		// that were forked after the first scan but before their parent received SIGSTOP.
+		TArray<pid_t> ToKill;
+		ToKill.Add(RootPid);
+		if (kill(RootPid, SIGSTOP) == -1 && errno != ESRCH)
+		{
+			int ErrNo = errno;
+			UE_LOGF(LogHAL, Warning, "TerminateProc KillTree: SIGSTOP to root pid %d failed (errno=%d, %s)",
+				static_cast<int32>(RootPid), ErrNo, strerror(ErrNo));
+		}
+
+		// ESRCH always needs to be handled as processes can end before we finish each of the passes/stages
+		for (int32 Pass = 0; Pass < 2; ++Pass)
+		{
+			// Capture process id and parent id
+			TArray<TPair<pid_t, pid_t>> AllProcs;
+
+			DIR* ProcDir = opendir("/proc");
+			if (ProcDir == nullptr)
+			{
+				int ErrNo = errno;
+				UE_LOGF(LogHAL, Warning, "TerminateProc KillTree: unable to open /proc (errno=%d, %s) -- will kill root process (pid %d) but cannot discover child processes.",
+					ErrNo, strerror(ErrNo), static_cast<int32>(RootPid));
+				break; // second pass would also fail
+			}
+			else
+			{
+				struct dirent* Entry;
+				while ((Entry = readdir(ProcDir)) != nullptr)
+				{
+					const char* Name = Entry->d_name;
+					if (!FCStringAnsi::IsNumeric(Name))
+					{
+						continue;
+					}
+
+					pid_t Pid = static_cast<pid_t>(FCStringAnsi::Atoi(Name));
+					if (Pid == 0)
+					{
+						continue;
+					}
+
+					// "/proc/" (6) + max PID digits (7) + "/status" (7) + '\0' (1) = 21; 48 is generous
+					char StatusPath[48];
+					FCStringAnsi::Sprintf(StatusPath, "/proc/%d/status", static_cast<int>(Pid));
+
+					FILE* StatusFile = fopen(StatusPath, "r");
+					if (StatusFile == nullptr)
+					{
+						continue;
+					}
+
+					pid_t ParentPid = -1;
+					char Line[128];
+					while (fgets(Line, sizeof(Line), StatusFile) != nullptr)
+					{
+						if (strncmp(Line, "PPid:", 5) == 0)
+						{
+							ParentPid = static_cast<pid_t>(FCStringAnsi::Atoi(Line + 5));
+							break;
+						}
+					}
+					fclose(StatusFile);
+
+					if (ParentPid >= 0)
+					{
+						AllProcs.Emplace(Pid, ParentPid);
+					}
+				}
+				closedir(ProcDir);
+			}
+
+			// BFS from every pid already in ToKill to discover new descendants.
+			// SIGSTOP each process immediately upon discovery to prevent it from forking.
+			for (int32 QueueIdx = 0; QueueIdx < ToKill.Num(); ++QueueIdx)
+			{
+				pid_t CurrentPid = ToKill[QueueIdx];
+				for (const TPair<pid_t, pid_t>& Proc : AllProcs)
+				{
+					if (Proc.Value == CurrentPid && !ToKill.Contains(Proc.Key))
+					{
+						ToKill.Add(Proc.Key);
+						if (kill(Proc.Key, SIGSTOP) == -1 && errno != ESRCH)
+						{
+							int ErrNo = errno;
+							UE_LOGF(LogHAL, Warning, "TerminateProc KillTree: SIGSTOP to pid %d failed (errno=%d, %s)",
+								static_cast<int32>(Proc.Key), ErrNo, strerror(ErrNo));
+						}
+					}
+				}
+			}
+		}
+
+		// Send SIGTERM to all the processes as a group before SIGCONT (which processes the SIGTERM) so the whole batch is kept together
+		for (int32 Index = ToKill.Num() - 1; Index >= 0; --Index)
+		{
+			if (kill(ToKill[Index], SIGTERM) == -1 && errno != ESRCH)
+			{
+				int ErrNo = errno;
+				UE_LOGF(LogHAL, Warning, "TerminateProc KillTree: SIGTERM to pid %d failed (errno=%d, %s)",
+					static_cast<int32>(ToKill[Index]), ErrNo, strerror(ErrNo));
+			}
+		}
+
+		// Resume all processes so the queued SIGTERM is delivered and they can exit gracefully.
+		for (int32 Index = ToKill.Num() - 1; Index >= 0; --Index)
+		{
+			if (kill(ToKill[Index], SIGCONT) == -1 && errno != ESRCH)
+			{
+				int ErrNo = errno;
+				UE_LOGF(LogHAL, Warning, "TerminateProc KillTree: SIGCONT to pid %d failed (errno=%d, %s)",
+					static_cast<int32>(ToKill[Index]), ErrNo, strerror(ErrNo));
+			}
+		}
+		return;
 	}
-	else if (ProcessHandle.Get() != -1)
+
+	if (RootPid != -1)
 	{
-		STUBBED("FUnixPlatformProcess::TerminateProc() : Terminating OpenProcess() handle not implemented");
+		int KillResult = kill(RootPid, SIGTERM);
+		check(KillResult != -1 || errno != EINVAL);
 	}
 }
 
@@ -2015,7 +2283,9 @@ bool FUnixPlatformProcess::GetProcReturnCode(FProcHandle& ProcHandle, int32* Ret
 	}
 	else if (ProcHandle.Get() != -1)
 	{
-		STUBBED("FUnixPlatformProcess::GetProcReturnCode() : Return code of OpenProcess() handle not implemented yet");
+		// Linux has no API to retrieve the exit code of a non-child process
+		UE_LOGF(LogHAL, Warning, "GetProcReturnCode: cannot retrieve exit code for OpenProcess() handle (pid %d) -- not the parent process.",
+			static_cast<int32>(ProcHandle.Get()));
 	}
 
 	return false;
diff --git a/Engine/Source/Runtime/Core/Public/Unix/UnixPlatformProcess.h b/Engine/Source/Runtime/Core/Public/Unix/UnixPlatformProcess.h
--- a/Engine/Source/Runtime/Core/Public/Unix/UnixPlatformProcess.h
+++ b/Engine/Source/Runtime/Core/Public/Unix/UnixPlatformProcess.h
@@ -111,25 +111,27 @@ struct FProcHandle
 	/** Child proc state set from CreateProc() call */
 	FProcState* 		ProcInfo;
 
-	/** Pid of external process opened with OpenProcess() call.
-	  * Added to FProcHandle so we don't have to special case FProcState with process
-	  * we can only check for running state, and even then the PID could be reused so
-	  * we don't ever want to terminate, etc.
+	/** Process start time captured during OpenProcess(), used with OpenedPid to
+	  * uniquely identify the process and detect PID reuse. Value is clock ticks
+	  * since boot from /proc/<pid>/stat field 22. Zero if unavailable.
 	  */
+	uint64				OpenedStartTime;
+
+	/** Pid of external process opened with OpenProcess() call. */
 	pid_t				OpenedPid;
 
 	constexpr FProcHandle()
-	:	ProcInfo(nullptr), OpenedPid(-1)
+	:	ProcInfo(nullptr), OpenedStartTime(0), OpenedPid(-1)
 	{
 	}
 
 	FProcHandle(FProcState* InHandle)
-	:	ProcInfo(InHandle), OpenedPid(-1)
+	:	ProcInfo(InHandle), OpenedStartTime(0), OpenedPid(-1)
 	{
 	}
 
 	FProcHandle(pid_t InProcPid)
-	:	ProcInfo(nullptr), OpenedPid(InProcPid)
+	:	ProcInfo(nullptr), OpenedStartTime(0), OpenedPid(InProcPid)
 	{
 	}
 
@@ -143,6 +145,7 @@ struct FProcHandle
 	inline void Reset()
 	{
 		ProcInfo = nullptr;
+		OpenedStartTime = 0;
 		OpenedPid = -1;
 	}
 
